Added RHICore module. Moved implementation of static uniform buffer bindings for D3D11.

diff --git a/Engine/Plugins/Lumin/MagicLeap/Source/MagicLeap/MagicLeap.Build.cs b/Engine/Plugins/Lumin/MagicLeap/Source/MagicLeap/MagicLeap.Build.cs
--- a/Engine/Plugins/Lumin/MagicLeap/Source/MagicLeap/MagicLeap.Build.cs
+++ b/Engine/Plugins/Lumin/MagicLeap/Source/MagicLeap/MagicLeap.Build.cs
@@ -41,6 +41,7 @@ namespace UnrealBuildTool.Rules
 					"EngineSettings",
 					"InputCore",
 					"RHI",
+					"RHICore",
 					"RenderCore",
 					"Renderer",
 					"Slate",
diff --git a/Engine/Plugins/Media/HLMedia/Source/HLMedia/HLMedia.Build.cs b/Engine/Plugins/Media/HLMedia/Source/HLMedia/HLMedia.Build.cs
--- a/Engine/Plugins/Media/HLMedia/Source/HLMedia/HLMedia.Build.cs
+++ b/Engine/Plugins/Media/HLMedia/Source/HLMedia/HLMedia.Build.cs
@@ -29,7 +29,8 @@ public class HLMedia : ModuleRules
                 "Projects",
                 "RenderCore",
                 "RHI",
-                "MediaUtils",
+				"RHICore",
+				"MediaUtils",
             });
 
         PrivateIncludePathModuleNames.AddRange(
diff --git a/Engine/Plugins/Media/PixelStreaming/Source/PixelStreaming/PixelStreaming.Build.cs b/Engine/Plugins/Media/PixelStreaming/Source/PixelStreaming/PixelStreaming.Build.cs
--- a/Engine/Plugins/Media/PixelStreaming/Source/PixelStreaming/PixelStreaming.Build.cs
+++ b/Engine/Plugins/Media/PixelStreaming/Source/PixelStreaming/PixelStreaming.Build.cs
@@ -102,7 +102,8 @@ namespace UnrealBuildTool.Rules
 				"Json",
 				"RenderCore",
                 "RHI",
-                "D3D11RHI",
+				"RHICore",
+				"D3D11RHI",
                 "Slate",
 				"SlateCore",
                 "AudioMixer",
diff --git a/Engine/Plugins/Media/WmfMedia/Source/WmfMedia/WmfMedia.Build.cs b/Engine/Plugins/Media/WmfMedia/Source/WmfMedia/WmfMedia.Build.cs
--- a/Engine/Plugins/Media/WmfMedia/Source/WmfMedia/WmfMedia.Build.cs
+++ b/Engine/Plugins/Media/WmfMedia/Source/WmfMedia/WmfMedia.Build.cs
@@ -27,6 +27,7 @@ namespace UnrealBuildTool.Rules
 					"Projects",
 					"RenderCore",
                     "RHI",
+					"RHICore",
                     "WmfMediaFactory"
                 });
 
diff --git a/Engine/Plugins/Runtime/Oculus/OculusOpenXR/Source/OculusOpenXRHMD/OculusOpenXRHMD.Build.cs b/Engine/Plugins/Runtime/Oculus/OculusOpenXR/Source/OculusOpenXRHMD/OculusOpenXRHMD.Build.cs
--- a/Engine/Plugins/Runtime/Oculus/OculusOpenXR/Source/OculusOpenXRHMD/OculusOpenXRHMD.Build.cs
+++ b/Engine/Plugins/Runtime/Oculus/OculusOpenXR/Source/OculusOpenXRHMD/OculusOpenXRHMD.Build.cs
@@ -43,7 +43,8 @@ namespace UnrealBuildTool.Rules
                     "Engine",
                     "InputCore",
                     "RHI",
-                    "RenderCore",
+					"RHICore",
+					"RenderCore",
                     "Renderer",
                     "Slate",
                     "SlateCore",
diff --git a/Engine/Plugins/Runtime/Oculus/OculusVR/Source/OculusHMD/OculusHMD.Build.cs b/Engine/Plugins/Runtime/Oculus/OculusVR/Source/OculusHMD/OculusHMD.Build.cs
--- a/Engine/Plugins/Runtime/Oculus/OculusVR/Source/OculusHMD/OculusHMD.Build.cs
+++ b/Engine/Plugins/Runtime/Oculus/OculusVR/Source/OculusHMD/OculusHMD.Build.cs
@@ -42,6 +42,7 @@ namespace UnrealBuildTool.Rules
 					"Engine",
 					"InputCore",
 					"RHI",
+					"RHICore",
 					"RenderCore",
 					"Renderer",
 					"Slate",
diff --git a/Engine/Plugins/Runtime/OpenXR/Source/OpenXRHMD/OpenXRHMD.Build.cs b/Engine/Plugins/Runtime/OpenXR/Source/OpenXRHMD/OpenXRHMD.Build.cs
--- a/Engine/Plugins/Runtime/OpenXR/Source/OpenXRHMD/OpenXRHMD.Build.cs
+++ b/Engine/Plugins/Runtime/OpenXR/Source/OpenXRHMD/OpenXRHMD.Build.cs
@@ -39,6 +39,7 @@ namespace UnrealBuildTool.Rules
                     "BuildSettings",
                     "InputCore",
 					"RHI",
+					"RHICore",
 					"RenderCore",
 					"Renderer",
 					"RenderCore",
diff --git a/Engine/Plugins/Runtime/Steam/SteamVR/Source/SteamVR/SteamVR.Build.cs b/Engine/Plugins/Runtime/Steam/SteamVR/Source/SteamVR/SteamVR.Build.cs
--- a/Engine/Plugins/Runtime/Steam/SteamVR/Source/SteamVR/SteamVR.Build.cs
+++ b/Engine/Plugins/Runtime/Steam/SteamVR/Source/SteamVR/SteamVR.Build.cs
@@ -31,6 +31,7 @@ namespace UnrealBuildTool.Rules
 					"CoreUObject",
 					"Engine",
 					"RHI",
+					"RHICore",
 					"RenderCore",
 					"Renderer",
                     "InputCore",
diff --git a/Engine/Plugins/Runtime/TextureShare/Source/TextureShare/TextureShare.Build.cs b/Engine/Plugins/Runtime/TextureShare/Source/TextureShare/TextureShare.Build.cs
--- a/Engine/Plugins/Runtime/TextureShare/Source/TextureShare/TextureShare.Build.cs
+++ b/Engine/Plugins/Runtime/TextureShare/Source/TextureShare/TextureShare.Build.cs
@@ -29,6 +29,7 @@ public class TextureShare : ModuleRules
 				"Core",
 				"Engine",
 				"RHI",
+				"RHICore",
 				"Renderer",
 				"RenderCore",
 				"TextureShareCore"
diff --git a/Engine/Plugins/Runtime/TextureShare/Source/TextureShareD3D11/TextureShareD3D11.Build.cs b/Engine/Plugins/Runtime/TextureShare/Source/TextureShareD3D11/TextureShareD3D11.Build.cs
--- a/Engine/Plugins/Runtime/TextureShare/Source/TextureShareD3D11/TextureShareD3D11.Build.cs
+++ b/Engine/Plugins/Runtime/TextureShare/Source/TextureShareD3D11/TextureShareD3D11.Build.cs
@@ -13,6 +13,7 @@ public class TextureShareD3D11 : ModuleRules
 							"Core",
 							"Engine",
 							"RHI",
+							"RHICore",
 							"RenderCore",
 							"D3D11RHI",
 						});
diff --git a/Engine/Plugins/Runtime/WindowsMixedReality/Source/WindowsMixedRealityHMD/WindowsMixedRealityHMD.Build.cs b/Engine/Plugins/Runtime/WindowsMixedReality/Source/WindowsMixedRealityHMD/WindowsMixedRealityHMD.Build.cs
--- a/Engine/Plugins/Runtime/WindowsMixedReality/Source/WindowsMixedRealityHMD/WindowsMixedRealityHMD.Build.cs
+++ b/Engine/Plugins/Runtime/WindowsMixedReality/Source/WindowsMixedRealityHMD/WindowsMixedRealityHMD.Build.cs
@@ -67,6 +67,7 @@ namespace UnrealBuildTool.Rules
                         "EngineSettings",
                         "InputCore",
 						"RHI",
+						"RHICore",
 						"RenderCore",
 						"Renderer",
 						"HeadMountedDisplay",
diff --git a/Engine/Plugins/Runtime/nDisplay/Source/DisplayClusterProjection/DisplayClusterProjection.Build.cs b/Engine/Plugins/Runtime/nDisplay/Source/DisplayClusterProjection/DisplayClusterProjection.Build.cs
--- a/Engine/Plugins/Runtime/nDisplay/Source/DisplayClusterProjection/DisplayClusterProjection.Build.cs
+++ b/Engine/Plugins/Runtime/nDisplay/Source/DisplayClusterProjection/DisplayClusterProjection.Build.cs
@@ -27,7 +27,8 @@ public class DisplayClusterProjection : ModuleRules
 				"MPCDI",
 				"Projects",
 				"RenderCore",
-				"RHI"
+				"RHI",
+				"RHICore"
 			});
 
 		if (Target.bBuildEditor == true)
diff --git a/Engine/Source/Runtime/RHICore/Private/RHICore.cpp b/Engine/Source/Runtime/RHICore/Private/RHICore.cpp
new file mode 100644
--- /dev/null
+++ b/Engine/Source/Runtime/RHICore/Private/RHICore.cpp
@@ -0,0 +1,7 @@
+// Copyright Epic Games, Inc. All Rights Reserved.
+
+#include "RHICore.h"
+#include "Modules/ModuleManager.h"
+
+DEFINE_LOG_CATEGORY(LogRHICore);
+IMPLEMENT_MODULE(FDefaultModuleImpl, RHICore);
\ No newline at end of file
diff --git a/Engine/Source/Runtime/RHICore/Private/RHICoreShader.cpp b/Engine/Source/Runtime/RHICore/Private/RHICoreShader.cpp
new file mode 100644
--- /dev/null
+++ b/Engine/Source/Runtime/RHICore/Private/RHICoreShader.cpp
@@ -0,0 +1,66 @@
+// Copyright Epic Games, Inc. All Rights Reserved.
+
+#include "RHICoreShader.h"
+
+#define RHI_VALIDATE_STATIC_UNIFORM_BUFFERS (!UE_BUILD_SHIPPING && !UE_BUILD_TEST)
+
+namespace UE
+{
+namespace RHICore
+{
+
+void ValidateStaticUniformBuffer(FRHIUniformBuffer* UniformBuffer, FUniformBufferStaticSlot Slot, uint32 ExpectedHash)
+{
+#if RHI_VALIDATE_STATIC_UNIFORM_BUFFERS
+	FUniformBufferStaticSlotRegistry& SlotRegistry = FUniformBufferStaticSlotRegistry::Get();
+
+	if (!UniformBuffer)
+	{
+		const FShaderParametersMetadata* ExpectedStructMetadata = FindUniformBufferStructByLayoutHash(ExpectedHash);
+
+		checkf(
+			ExpectedStructMetadata,
+			TEXT("Shader is requesting a uniform buffer at slot %s with hash '%u', but a reverse lookup of the hash can't find it. The shader cache may be out of date."),
+			*SlotRegistry.GetDebugDescription(Slot), ExpectedHash);
+
+		const EUniformBufferBindingFlags BindingFlags = ExpectedStructMetadata->GetBindingFlags();
+
+		checkf(EnumHasAnyFlags(BindingFlags, EUniformBufferBindingFlags::Static),
+			TEXT("Shader requested a global uniform buffer of type '%s' at static slot '%s', but it is not registered with the Global binding flag. The shader cache may be out of date."),
+			ExpectedStructMetadata->GetShaderVariableName(), *SlotRegistry.GetDebugDescription(Slot));
+
+		// Structs can be bound both globally or per-shader, effectively leaving it up to the user to choose which to bind.
+		// But that also means we can't validate existence at the global level.
+		if (!EnumHasAnyFlags(BindingFlags, EUniformBufferBindingFlags::Shader))
+		{
+			UE_LOG(LogRHICore, Fatal,
+				TEXT("Shader requested a global uniform buffer of type '%s' at static slot '%s', but it was null. The uniform buffer should ")
+				TEXT("be bound using RHICmdList.SetGlobalUniformBuffers() or passed into an RDG pass using SHADER_PARAMETER_STRUCT_REF() or ")
+				TEXT("SHADER_PARAMETER_RDG_UNIFORM_BUFFER()."),
+				ExpectedStructMetadata->GetShaderVariableName(), *SlotRegistry.GetDebugDescription(Slot));
+		}
+	}
+	else
+	{
+		const FRHIUniformBufferLayout& Layout = UniformBuffer->GetLayout();
+
+		if (Layout.GetHash() != ExpectedHash)
+		{
+			const FShaderParametersMetadata* ExpectedStructMetadata = FindUniformBufferStructByLayoutHash(ExpectedHash);
+
+			checkf(
+				ExpectedStructMetadata,
+				TEXT("Shader is requesting uniform buffer '%s' at slot %s with hash '%u', but a reverse lookup of the hash can't find it. The shader cache may be out of date."),
+				*Layout.GetDebugName(), *SlotRegistry.GetDebugDescription(Slot), ExpectedHash);
+
+			checkf(
+				false,
+				TEXT("Shader attempted to bind uniform buffer '%s' at slot %s with hash '%u', but the shader expected '%s' with hash '%u'."),
+				*Layout.GetDebugName(), *SlotRegistry.GetDebugDescription(Slot), ExpectedHash, ExpectedStructMetadata->GetShaderVariableName(), Layout.GetHash());
+		}
+	}
+#endif
+}
+
+} //! RHICore
+} //! UE
\ No newline at end of file
diff --git a/Engine/Source/Runtime/RHICore/Public/RHICore.h b/Engine/Source/Runtime/RHICore/Public/RHICore.h
new file mode 100644
--- /dev/null
+++ b/Engine/Source/Runtime/RHICore/Public/RHICore.h
@@ -0,0 +1,5 @@
+// Copyright Epic Games, Inc. All Rights Reserved.
+
+#pragma once
+
+RHICORE_API DECLARE_LOG_CATEGORY_EXTERN(LogRHICore, Log, VeryVerbose);
\ No newline at end of file
diff --git a/Engine/Source/Runtime/RHICore/Public/RHICoreShader.h b/Engine/Source/Runtime/RHICore/Public/RHICoreShader.h
new file mode 100644
--- /dev/null
+++ b/Engine/Source/Runtime/RHICore/Public/RHICoreShader.h
@@ -0,0 +1,45 @@
+// Copyright Epic Games, Inc. All Rights Reserved.
+
+#pragma once
+
+#include "RHICore.h"
+#include "ShaderCore.h"
+
+namespace UE
+{
+namespace RHICore
+{
+
+/** Validates that the uniform buffer at the requested static slot. */
+extern RHICORE_API void ValidateStaticUniformBuffer(FRHIUniformBuffer* UniformBuffer, FUniformBufferStaticSlot Slot, uint32 ExpectedHash);
+
+template <typename TRHIContext, typename TRHIShader>
+void ApplyStaticUniformBuffers(
+	TRHIContext* CommandContext,
+	TRHIShader* Shader,
+	const TArray<FUniformBufferStaticSlot>& Slots,
+	const TArray<uint32>& LayoutHashes,
+	const TArray<FRHIUniformBuffer*>& UniformBuffers)
+{
+	checkf(LayoutHashes.Num() == Slots.Num(), TEXT("Shader %s, LayoutHashes %d, Slots %d"),
+		Shader->GetShaderName(), LayoutHashes.Num(), Slots.Num());
+
+	for (int32 BufferIndex = 0; BufferIndex < Slots.Num(); ++BufferIndex)
+	{
+		const FUniformBufferStaticSlot Slot = Slots[BufferIndex];
+
+		if (IsUniformBufferStaticSlotValid(Slot))
+		{
+			FRHIUniformBuffer* Buffer = UniformBuffers[Slot];
+			RHICore::ValidateStaticUniformBuffer(Buffer, Slot, LayoutHashes[BufferIndex]);
+
+			if (Buffer)
+			{
+				CommandContext->RHISetShaderUniformBuffer(Shader, BufferIndex, Buffer);
+			}
+		}
+	}
+}
+
+} //! RHICore
+} //! UE
\ No newline at end of file
diff --git a/Engine/Source/Runtime/RHICore/RHICore.Build.cs b/Engine/Source/Runtime/RHICore/RHICore.Build.cs
new file mode 100644
--- /dev/null
+++ b/Engine/Source/Runtime/RHICore/RHICore.Build.cs
@@ -0,0 +1,14 @@
+// Copyright Epic Games, Inc. All Rights Reserved.
+
+using UnrealBuildTool;
+using System;
+
+public class RHICore : ModuleRules
+{
+	public RHICore(ReadOnlyTargetRules Target) : base(Target)
+	{
+		PublicDependencyModuleNames.AddRange(new string[] { "RHI" });
+		PublicDependencyModuleNames.AddRange(new string[] { "RenderCore" });
+		PrivateDependencyModuleNames.AddRange(new string[] { "Core", "Projects", "RHI", "ApplicationCore", "TraceLog" });
+	}
+}
diff --git a/Engine/Source/Runtime/Windows/D3D11RHI/D3D11RHI.Build.cs b/Engine/Source/Runtime/Windows/D3D11RHI/D3D11RHI.Build.cs
--- a/Engine/Source/Runtime/Windows/D3D11RHI/D3D11RHI.Build.cs
+++ b/Engine/Source/Runtime/Windows/D3D11RHI/D3D11RHI.Build.cs
@@ -20,6 +20,7 @@ public class D3D11RHI : ModuleRules
 				"CoreUObject",
 				"Engine",
 				"RHI",
+				"RHICore",
 				"RenderCore"
 			}
 			);
diff --git a/Engine/Source/Runtime/Windows/D3D11RHI/Private/D3D11Commands.cpp b/Engine/Source/Runtime/Windows/D3D11RHI/Private/D3D11Commands.cpp
--- a/Engine/Source/Runtime/Windows/D3D11RHI/Private/D3D11Commands.cpp
+++ b/Engine/Source/Runtime/Windows/D3D11RHI/Private/D3D11Commands.cpp
@@ -195,11 +195,11 @@ void FD3D11DynamicRHI::RHISetGraphicsPipelineState(FRHIGraphicsPipelineState* Gr
 
 	if (bApplyAdditionalState)
 	{
-		ApplyGlobalUniformBuffers(static_cast<FD3D11VertexShader*>(PsoInit.BoundShaderState.VertexShaderRHI));
-		ApplyGlobalUniformBuffers(static_cast<FD3D11HullShader*>(PsoInit.BoundShaderState.HullShaderRHI));
-		ApplyGlobalUniformBuffers(static_cast<FD3D11DomainShader*>(PsoInit.BoundShaderState.DomainShaderRHI));
-		ApplyGlobalUniformBuffers(static_cast<FD3D11GeometryShader*>(PsoInit.BoundShaderState.GeometryShaderRHI));
-		ApplyGlobalUniformBuffers(static_cast<FD3D11PixelShader*>(PsoInit.BoundShaderState.PixelShaderRHI));
+		ApplyStaticUniformBuffers(static_cast<FD3D11VertexShader*>(PsoInit.BoundShaderState.VertexShaderRHI));
+		ApplyStaticUniformBuffers(static_cast<FD3D11HullShader*>(PsoInit.BoundShaderState.HullShaderRHI));
+		ApplyStaticUniformBuffers(static_cast<FD3D11DomainShader*>(PsoInit.BoundShaderState.DomainShaderRHI));
+		ApplyStaticUniformBuffers(static_cast<FD3D11GeometryShader*>(PsoInit.BoundShaderState.GeometryShaderRHI));
+		ApplyStaticUniformBuffers(static_cast<FD3D11PixelShader*>(PsoInit.BoundShaderState.PixelShaderRHI));
 	}
 
 	// Store the PSO's primitive (after since IRHICommandContext::RHISetGraphicsPipelineState sets the BSS)
@@ -216,7 +216,7 @@ void FD3D11DynamicRHI::RHISetComputeShader(FRHIComputeShader* ComputeShaderRHI)
 		ClearAllShaderResourcesForFrequency<SF_Compute>();
 	}
 
-	ApplyGlobalUniformBuffers(ComputeShader);
+	ApplyStaticUniformBuffers(ComputeShader);
 }
 
 void FD3D11DynamicRHI::RHIDispatchComputeShader(uint32 ThreadGroupCountX, uint32 ThreadGroupCountY, uint32 ThreadGroupCountZ) 
@@ -616,11 +616,11 @@ void FD3D11DynamicRHI::RHISetShaderSampler(FRHIComputeShader* ComputeShaderRHI,u
 
 void FD3D11DynamicRHI::RHISetGlobalUniformBuffers(const FUniformBufferStaticBindings& InUniformBuffers)
 {
-	FMemory::Memzero(GlobalUniformBuffers.GetData(), GlobalUniformBuffers.Num() * sizeof(FRHIUniformBuffer*));
+	FMemory::Memzero(StaticUniformBuffers.GetData(), StaticUniformBuffers.Num() * sizeof(FRHIUniformBuffer*));
 
 	for (int32 Index = 0; Index < InUniformBuffers.GetUniformBufferCount(); ++Index)
 	{
-		GlobalUniformBuffers[InUniformBuffers.GetSlot(Index)] = InUniformBuffers.GetUniformBuffer(Index);
+		StaticUniformBuffers[InUniformBuffers.GetSlot(Index)] = InUniformBuffers.GetUniformBuffer(Index);
 	}
 }
 
diff --git a/Engine/Source/Runtime/Windows/D3D11RHI/Private/D3D11Device.cpp b/Engine/Source/Runtime/Windows/D3D11RHI/Private/D3D11Device.cpp
--- a/Engine/Source/Runtime/Windows/D3D11RHI/Private/D3D11Device.cpp
+++ b/Engine/Source/Runtime/Windows/D3D11RHI/Private/D3D11Device.cpp
@@ -241,7 +241,7 @@ FD3D11DynamicRHI::FD3D11DynamicRHI(IDXGIFactory1* InDXGIFactory1,D3D_FEATURE_LEV
 		DirtyUniformBuffers[Frequency] = 0;
 	}
 
-	GlobalUniformBuffers.AddZeroed(FUniformBufferStaticSlotRegistry::Get().GetSlotCount());
+	StaticUniformBuffers.AddZeroed(FUniformBufferStaticSlotRegistry::Get().GetSlotCount());
 }
 
 FD3D11DynamicRHI::~FD3D11DynamicRHI()
diff --git a/Engine/Source/Runtime/Windows/D3D11RHI/Private/D3D11RHIPrivate.h b/Engine/Source/Runtime/Windows/D3D11RHI/Private/D3D11RHIPrivate.h
--- a/Engine/Source/Runtime/Windows/D3D11RHI/Private/D3D11RHIPrivate.h
+++ b/Engine/Source/Runtime/Windows/D3D11RHI/Private/D3D11RHIPrivate.h
@@ -29,6 +29,7 @@ DECLARE_LOG_CATEGORY_EXTERN(LogD3D11RHI, Log, All);
 #include "D3D11ConstantBuffer.h"
 #include "D3D11StateCache.h"
 #include "RHIValidationCommon.h"
+#include "RHICoreShader.h"
 
 #include "ShaderCompiler.h"
 
@@ -878,15 +879,15 @@ protected:
 	uint16 DirtyUniformBuffers[SF_NumStandardFrequencies];
 
 	template <typename TRHIShader>
-	void ApplyGlobalUniformBuffers(TRHIShader* Shader)
+	void ApplyStaticUniformBuffers(TRHIShader* Shader)
 	{
 		if (Shader)
 		{
-			::ApplyGlobalUniformBuffers(this, Shader, Shader->StaticSlots, Shader->ShaderResourceTable.ResourceTableLayoutHashes, GlobalUniformBuffers);
+			UE::RHICore::ApplyStaticUniformBuffers(this, Shader, Shader->StaticSlots, Shader->ShaderResourceTable.ResourceTableLayoutHashes, StaticUniformBuffers);
 		}
 	}
 
-	TArray<FRHIUniformBuffer*> GlobalUniformBuffers;
+	TArray<FRHIUniformBuffer*> StaticUniformBuffers;
 
 	/** Tracks the current depth stencil access type. */
 	FExclusiveDepthStencil CurrentDSVAccessType;
