Fix non-unity

diff --git a/Engine/Source/Runtime/RHICore/Public/RHICore.h b/Engine/Source/Runtime/RHICore/Public/RHICore.h
--- a/Engine/Source/Runtime/RHICore/Public/RHICore.h
+++ b/Engine/Source/Runtime/RHICore/Public/RHICore.h
@@ -2,4 +2,6 @@
 
 #pragma once
 
+#include "Logging/LogMacros.h"
+
 RHICORE_API DECLARE_LOG_CATEGORY_EXTERN(LogRHICore, Log, VeryVerbose);
\ No newline at end of file
