Fix for hitting check() in RHI allocator

diff --git a/Engine/Source/Runtime/RHICore/Private/RHIPoolAllocator.cpp b/Engine/Source/Runtime/RHICore/Private/RHIPoolAllocator.cpp
--- a/Engine/Source/Runtime/RHICore/Private/RHIPoolAllocator.cpp
+++ b/Engine/Source/Runtime/RHICore/Private/RHIPoolAllocator.cpp
@@ -351,14 +351,15 @@ void FRHIMemoryPool::TryClear(FRHIPoolAllocator* InAllocator, uint32 InMaxCopySi
 
 int32 FRHIMemoryPool::FindFreeBlock(uint32 InSizeInBytes, uint32 InAllocationAlignment) const
 {
+	uint32 AlignedSize = GetAlignedSize(InSizeInBytes, PoolAlignment, InAllocationAlignment);
+
 	// Early out if total free size doesn't fit
-	if (FreeSize < InSizeInBytes)
+	if (FreeSize < AlignedSize)
 	{
 		return INDEX_NONE;
 	}
 
 	// TODO: lowerbound search
-	uint32 AlignedSize = GetAlignedSize(InSizeInBytes, PoolAlignment, InAllocationAlignment);
 
 	// Check all the free blocks (sorted by size, so take best fit)
 	for (int32 FreeBlockIndex = 0; FreeBlockIndex < FreeBlocks.Num(); ++FreeBlockIndex)
