Refactored platform agnostic page allocator into RHI core.

diff --git a/Engine/Source/Runtime/RHICore/Private/RHIPageAllocator.cpp b/Engine/Source/Runtime/RHICore/Private/RHIPageAllocator.cpp
new file mode 100644
--- /dev/null
+++ b/Engine/Source/Runtime/RHICore/Private/RHIPageAllocator.cpp
@@ -0,0 +1,398 @@
+
+#include "RHIPageAllocator.h"
+
+void FRHIPageAllocator::Init()
+{
+	check(MaxSpanCount == PageCount + 2);
+
+	PageToSpanStart.AddDefaulted(PageCount + 1);
+	PageToSpanEnd.AddDefaulted(PageCount + 1);
+
+	PageSpans.AddDefaulted(MaxSpanCount);
+	UnusedSpanList.AddDefaulted(MaxSpanCount);
+
+	Reset();
+}
+
+void FRHIPageAllocator::Reset()
+{
+	FreePageCount = PageCount;
+	AllocationCount = 0;
+
+	// Initialize the unused span index pool with MaxSpanCount entries
+	for (uint32 i = 0; i < MaxSpanCount; i++)
+	{
+		UnusedSpanList[i] = MaxSpanCount - 1 - i;
+	}
+	UnusedSpanListCount = MaxSpanCount;
+
+	// Allocate the head and tail spans (dummy spans), and a span between them covering the entire range
+	int32 HeadSpanIndex = AllocSpan();
+	int32 TailSpanIndex = AllocSpan();
+	check(HeadSpanIndex == FREE_SPAN_LIST_HEAD_INDEX);
+	check(TailSpanIndex == FREE_SPAN_LIST_TAIL_INDEX);
+
+	if (PageCount > 0)
+	{
+		int32 FirstFreeNodeIndex = AllocSpan();
+
+		// Allocate head and tail nodes (0 and 1)
+		for (int32 i = 0; i < 2; i++)
+		{
+			PageSpans[i].StartPageIndex = 0;
+			PageSpans[i].Count = 0;
+			PageSpans[i].PrevSpanIndex = INVALID_INDEX;
+			PageSpans[i].NextSpanIndex = INVALID_INDEX;
+			PageSpans[i].bAllocated = false;
+		}
+		PageSpans[HeadSpanIndex].NextSpanIndex = FirstFreeNodeIndex;
+		PageSpans[TailSpanIndex].PrevSpanIndex = FirstFreeNodeIndex;
+
+		// First Node
+		PageSpans[FirstFreeNodeIndex].StartPageIndex = 0;
+		PageSpans[FirstFreeNodeIndex].Count = PageCount;
+		PageSpans[FirstFreeNodeIndex].PrevSpanIndex = HeadSpanIndex;
+		PageSpans[FirstFreeNodeIndex].NextSpanIndex = TailSpanIndex;
+		PageSpans[FirstFreeNodeIndex].bAllocated = false;
+
+		// Initialize the page->span mapping
+		for (uint32 i = 0; i < PageCount + 1; i++)
+		{
+			PageToSpanStart[i] = INVALID_INDEX;
+			PageToSpanEnd[i] = INVALID_INDEX;
+		}
+		PageToSpanStart[0] = FirstFreeNodeIndex;
+		PageToSpanEnd[PageCount] = FirstFreeNodeIndex;
+	}
+	else
+	{
+		PageSpans[HeadSpanIndex].NextSpanIndex = TailSpanIndex;
+		PageSpans[TailSpanIndex].PrevSpanIndex = HeadSpanIndex;
+	}
+}
+
+FRHIPageAllocator::Handle FRHIPageAllocator::AllocPages(int32 Count, const TCHAR* DebugName)
+{
+	uint32 NumPagesAllocated = 0;
+	Handle SpanIndex = AllocPagesInternal(Count, NumPagesAllocated, false, DebugName);
+	return SpanIndex;
+}
+
+FRHIPageAllocator::Handle FRHIPageAllocator::AllocPagesPartial(int32 Count, uint32& NumPagesAllocated, const TCHAR* DebugName)
+{
+	return AllocPagesInternal(Count, NumPagesAllocated, true, DebugName);
+}
+
+FRHIPageAllocator::Handle FRHIPageAllocator::AllocPagesInternal(int32 Count, uint32& NumPagesAllocated, bool bAllowPartialAlloc, const TCHAR* DebugName)
+{
+	NumPagesAllocated = 0;
+	if (bAllowPartialAlloc && FreePageCount < Count)
+	{
+		// If we're allowing partial allocs and we run out of pages, allocate all the remaining pages
+		Count = FreePageCount;
+	}
+
+	if (Count > FreePageCount || Count == 0)
+	{
+		return INVALID_INDEX;
+	}
+	NumPagesAllocated = Count;
+
+	// Allocate spans from the free list head
+	int32 NumPagesToFind = Count;
+	int32 FoundPages = 0;
+	FPageSpan& HeadSpan = PageSpans[FREE_SPAN_LIST_HEAD_INDEX];
+	int32 StartIndex = HeadSpan.NextSpanIndex;
+	int32 SpanIndex = StartIndex;
+	while (SpanIndex > FREE_SPAN_LIST_TAIL_INDEX)
+	{
+		FPageSpan& Span = PageSpans[SpanIndex];
+		if (NumPagesToFind <= Span.Count)
+		{
+			// Span is too big, so split it
+			if (Span.Count > NumPagesToFind)
+			{
+				SplitSpan(SpanIndex, NumPagesToFind);
+			}
+			check(NumPagesToFind == Span.Count);
+
+			// Move the head to point to the next free span
+			if (HeadSpan.NextSpanIndex >= 0)
+			{
+				PageSpans[HeadSpan.NextSpanIndex].PrevSpanIndex = INVALID_INDEX;
+			}
+			HeadSpan.NextSpanIndex = Span.NextSpanIndex;
+			if (Span.NextSpanIndex >= 0)
+			{
+				PageSpans[Span.NextSpanIndex].PrevSpanIndex = FREE_SPAN_LIST_HEAD_INDEX;
+			}
+			Span.NextSpanIndex = INVALID_INDEX;
+		}
+		Span.bAllocated = true;
+		NumPagesToFind -= Span.Count;
+		SpanIndex = Span.NextSpanIndex;
+	}
+	check(NumPagesToFind == 0);
+	FreePageCount -= Count;
+#if UE_BUILD_DEBUG
+	Validate();
+#endif
+	AllocationCount++;
+	return StartIndex;
+}
+
+void FRHIPageAllocator::SplitSpan(int32 InSpanIndex, int32 InPageCount)
+{
+	FPageSpan& Span = PageSpans[InSpanIndex];
+	check(InPageCount <= Span.Count);
+	if (InPageCount < Span.Count)
+	{
+		int32 NewSpanIndex = AllocSpan();
+		FPageSpan& NewSpan = PageSpans[NewSpanIndex];
+		NewSpan.NextSpanIndex = Span.NextSpanIndex;
+		NewSpan.PrevSpanIndex = InSpanIndex;
+		NewSpan.Count = Span.Count - InPageCount;
+		NewSpan.StartPageIndex = Span.StartPageIndex + InPageCount;
+		NewSpan.bAllocated = Span.bAllocated;
+		Span.Count = InPageCount;
+		Span.NextSpanIndex = NewSpanIndex;
+		if (NewSpan.NextSpanIndex >= 0)
+		{
+			PageSpans[NewSpan.NextSpanIndex].PrevSpanIndex = NewSpanIndex;
+		}
+
+		// Update the PageToSpan mappings
+		PageToSpanEnd[NewSpan.StartPageIndex] = InSpanIndex;
+		PageToSpanStart[NewSpan.StartPageIndex] = NewSpanIndex;
+		PageToSpanEnd[NewSpan.StartPageIndex + NewSpan.Count] = NewSpanIndex;
+	}
+}
+
+void FRHIPageAllocator::MergeSpans(int32 SpanIndex0, int32 SpanIndex1, const bool bKeepSpan1)
+{
+	FPageSpan& Span0 = PageSpans[SpanIndex0];
+	FPageSpan& Span1 = PageSpans[SpanIndex1];
+	check(Span0.StartPageIndex + Span0.Count == Span1.StartPageIndex);
+	check(Span0.bAllocated == Span1.bAllocated);
+	check(Span0.NextSpanIndex == SpanIndex1);
+	check(Span1.PrevSpanIndex == SpanIndex0);
+
+	int32 SpanIndexToKeep = bKeepSpan1 ? SpanIndex1 : SpanIndex0;
+	int32 SpanIndexToRemove = bKeepSpan1 ? SpanIndex0 : SpanIndex1;
+
+	// Update the PageToSpan mappings
+	PageToSpanStart[Span0.StartPageIndex] = SpanIndexToKeep;
+	PageToSpanStart[Span1.StartPageIndex] = INVALID_INDEX;
+	PageToSpanEnd[Span0.StartPageIndex + Span0.Count] = INVALID_INDEX; // Should match Span1.StartIndex
+	PageToSpanEnd[Span1.StartPageIndex + Span1.Count] = SpanIndexToKeep;
+	if (bKeepSpan1)
+	{
+		Span1.StartPageIndex = Span0.StartPageIndex;
+		Span1.Count += Span0.Count;
+	}
+	else
+	{
+		Span0.Count += Span1.Count;
+	}
+
+	Unlink(SpanIndexToRemove);
+	ReleaseSpan(SpanIndexToRemove);
+}
+
+// Inserts a span after an existing span. The span to insert must be unlinked
+void FRHIPageAllocator::InsertAfter(int32 InsertPosition, int32 InsertSpanIndex)
+{
+	check(InsertPosition >= 0);
+	check(InsertSpanIndex >= 0);
+	FPageSpan& SpanAtPos = PageSpans[InsertPosition];
+	FPageSpan& SpanToInsert = PageSpans[InsertSpanIndex];
+	check(!SpanToInsert.IsLinked());
+
+	// Connect Span0's next node with the inserted node
+	SpanToInsert.NextSpanIndex = SpanAtPos.NextSpanIndex;
+	if (SpanAtPos.NextSpanIndex >= 0)
+	{
+		PageSpans[SpanAtPos.NextSpanIndex].PrevSpanIndex = InsertSpanIndex;
+	}
+	// Connect the two nodes
+	SpanAtPos.NextSpanIndex = InsertSpanIndex;
+	SpanToInsert.PrevSpanIndex = InsertPosition;
+}
+
+// Inserts a span after an existing span. The span to insert must be unlinked
+void FRHIPageAllocator::InsertBefore(int32 InsertPosition, int32 InsertSpanIndex)
+{
+	check(InsertPosition > 0); // Can't insert before the head
+	check(InsertSpanIndex >= 0);
+	FPageSpan& SpanAtPos = PageSpans[InsertPosition];
+	FPageSpan& SpanToInsert = PageSpans[InsertSpanIndex];
+	check(!SpanToInsert.IsLinked());
+
+	// Connect Span0's prev node with the inserted node
+	SpanToInsert.PrevSpanIndex = SpanAtPos.PrevSpanIndex;
+	if (SpanAtPos.PrevSpanIndex >= 0)
+	{
+		PageSpans[SpanAtPos.PrevSpanIndex].NextSpanIndex = InsertSpanIndex;
+	}
+	// Connect the two nodes
+	SpanAtPos.PrevSpanIndex = InsertSpanIndex;
+	SpanToInsert.NextSpanIndex = InsertPosition;
+}
+
+void FRHIPageAllocator::Unlink(int32 SpanIndex)
+{
+	FPageSpan& Span = PageSpans[SpanIndex];
+	check(SpanIndex != FREE_SPAN_LIST_HEAD_INDEX);
+	if (Span.PrevSpanIndex != INVALID_INDEX)
+	{
+		PageSpans[Span.PrevSpanIndex].NextSpanIndex = Span.NextSpanIndex;
+	}
+	if (Span.NextSpanIndex != INVALID_INDEX)
+	{
+		PageSpans[Span.NextSpanIndex].PrevSpanIndex = Span.PrevSpanIndex;
+	}
+	Span.PrevSpanIndex = INVALID_INDEX;
+	Span.NextSpanIndex = INVALID_INDEX;
+}
+
+void FRHIPageAllocator::FreePages(Handle SpanIndex)
+{
+	if (SpanIndex == INVALID_INDEX)
+	{
+		return;
+	}
+	check(AllocationCount > 0);
+	// Find the right span with which to merge this
+	while (SpanIndex != INVALID_INDEX)
+	{
+		FPageSpan& FreedSpan = PageSpans[SpanIndex];
+		check(FreedSpan.bAllocated);
+		FreePageCount += FreedSpan.Count;
+		int32 NextSpanIndex = FreedSpan.NextSpanIndex;
+		FreedSpan.bAllocated = false;
+		if (!MergeFreeSpanIfPossible(SpanIndex))
+		{
+			// If we can't merge this span, just unlink and add it to the head (or tail)
+			Unlink(SpanIndex);
+
+#if 1
+			// This heuristic gave slightly reduced fragmentation in testing
+			const bool bAddToHead = (FreedSpan.Count >= 24);
+#else
+			const bool bAddToHead = true;
+#endif
+			if (bAddToHead)
+			{
+				InsertAfter(FREE_SPAN_LIST_HEAD_INDEX, SpanIndex);
+			}
+			else
+			{
+				InsertBefore(FREE_SPAN_LIST_TAIL_INDEX, SpanIndex);
+			}
+		}
+		SpanIndex = NextSpanIndex;
+	}
+	AllocationCount--;
+
+#if UE_BUILD_DEBUG
+	Validate();
+#endif
+}
+
+// Generates a flag array of pages for a given SpanIndex
+void FRHIPageAllocator::GetPageArray(Handle SpanIndex, TArray<uint32>& PagesOut, uint32 PageOffset, bool bAppend)
+{
+	if (!bAppend)
+	{
+		PagesOut.Reset();
+	}
+	int32 Index = SpanIndex;
+	int32 ArrayCapacity = 0;
+	while (Index != INVALID_INDEX)
+	{
+		FPageSpan& Span = PageSpans[Index];
+		int32 PageIndex = Span.StartPageIndex;
+		ArrayCapacity += Span.Count;
+		PagesOut.Reserve(ArrayCapacity);
+		for (int32 i = 0; i < Span.Count; i++)
+		{
+			PagesOut.Add((uint32)PageIndex + PageOffset);
+			PageIndex++;
+		}
+		Index = PageSpans[Index].NextSpanIndex;
+	}
+}
+
+// Generates an array of contiguous page ranges
+void FRHIPageAllocator::GetRangeArray(Handle SpanIndex, TArray<FVRamPageRange>& RangesOut, uint32 PageOffset, bool bAppend)
+{
+	if (!bAppend)
+	{
+		RangesOut.Reset();
+	}
+
+	int32 Index = SpanIndex;
+	int32 ArrayCapacity = 0;
+	while (Index != INVALID_INDEX)
+	{
+		FPageSpan& Span = PageSpans[Index];
+		FVRamPageRange Range;
+		Range.StartIndex = Span.StartPageIndex + PageOffset;
+		Range.Count = Span.Count;
+		RangesOut.Add(Range);
+
+		Index = PageSpans[Index].NextSpanIndex;
+	}
+}
+
+bool FRHIPageAllocator::MergeFreeSpanIfPossible(int32 SpanIndex)
+{
+	FPageSpan& Span = PageSpans[SpanIndex];
+	check(!Span.bAllocated);
+	bool bMerged = false;
+
+	// Can we merge this span with an existing one to the left?
+	int32 AdjSpanIndexPrev = PageToSpanEnd[Span.StartPageIndex];
+	if (AdjSpanIndexPrev >= 0 && !PageSpans[AdjSpanIndexPrev].bAllocated)
+	{
+		Unlink(SpanIndex);
+		InsertAfter(AdjSpanIndexPrev, SpanIndex);
+		MergeSpans(AdjSpanIndexPrev, SpanIndex, true);
+		bMerged = true;
+	}
+	// Can we merge this span with an existing free one to the right?
+	int32 AdjSpanIndexNext = PageToSpanStart[Span.StartPageIndex + Span.Count];
+	if (AdjSpanIndexNext >= 0 && !PageSpans[AdjSpanIndexNext].bAllocated)
+	{
+		Unlink(SpanIndex);
+		InsertBefore(AdjSpanIndexNext, SpanIndex);
+		MergeSpans(SpanIndex, AdjSpanIndexNext, false);
+		bMerged = true;
+	}
+	return bMerged;
+}
+
+void FRHIPageAllocator::Validate()
+{
+#if UE_BUILD_DEBUG
+	// Check the mappings are valid
+	for (uint32 i = 0; i < PageCount; i++)
+	{
+		check(PageToSpanStart[i] == INVALID_INDEX || PageSpans[PageToSpanStart[i]].StartPageIndex == i);
+		check(PageToSpanEnd[i] == INVALID_INDEX || PageSpans[PageToSpanEnd[i]].StartPageIndex + PageSpans[PageToSpanEnd[i]].Count == i);
+	}
+
+	// Count free pages
+	uint32 FreeCount = 0;
+	int32 PrevIndex = FREE_SPAN_LIST_HEAD_INDEX;
+	for (int32 index = GetFirstSpanIndex(); index != INVALID_INDEX; index = PageSpans[index].NextSpanIndex)
+	{
+		FPageSpan& Span = PageSpans[index];
+		check(Span.PrevSpanIndex == PrevIndex);
+		PrevIndex = index;
+		FreeCount += Span.Count;
+	}
+	check(FreeCount <= PageCount);
+	check(FreeCount == FreePageCount);
+#endif
+}
\ No newline at end of file
diff --git a/Engine/Source/Runtime/RHICore/Public/RHIPageAllocator.h b/Engine/Source/Runtime/RHICore/Public/RHIPageAllocator.h
new file mode 100644
--- /dev/null
+++ b/Engine/Source/Runtime/RHICore/Public/RHIPageAllocator.h
@@ -0,0 +1,167 @@
+// Copyright Epic Games, Inc. All Rights Reserved.
+
+#include "RHI.h"
+
+#pragma once
+
+struct FVRamPageRange
+{
+	uint32 StartIndex = 0;
+	uint32 Count = 0;
+};
+
+/** A non-contiguous span allocator for pages. */
+class FRHIPageAllocator
+{
+	struct FPageSpan
+	{
+		const bool IsLinked() { return (NextSpanIndex >= 0 || PrevSpanIndex >= 0); }
+
+		int32 StartPageIndex = -1;
+		int32 Count = 0;
+		int32 NextSpanIndex = -1;
+		int32 PrevSpanIndex = 0;
+		bool bAllocated = false;
+	};
+
+	static const int32 FREE_SPAN_LIST_HEAD_INDEX = 0;
+	static const int32 FREE_SPAN_LIST_TAIL_INDEX = 1;
+public:
+	static const int32 INVALID_INDEX = -1;
+	typedef int32 Handle;
+
+	FRHIPageAllocator()
+		: MaxSpanCount(514)
+		, PageCount(512)
+	{
+		Init();
+	}
+
+	FRHIPageAllocator(uint32 InPageCount)
+		: MaxSpanCount(InPageCount + 2)
+		, PageCount(InPageCount)
+	{
+		Init();
+	}
+
+	void Reset();
+
+	// Allocates pages, returning a span index (or INVALID_INDEX on failure), allows partial allocations, but fails if zero pages were available
+	Handle AllocPagesPartial(int32 Count, uint32& NumPagesAllocated, const TCHAR* DebugName);
+
+	// Allocates pages, returning a span index (or INVALID_INDEX on failure)
+	Handle AllocPages(int32 Count, const TCHAR* DebugName = nullptr);
+
+	// Frees previously allocated pages by span index
+	void FreePages(Handle SpanIndex);
+
+	// Generates a flag array of pages for a given SpanIndex
+	void GetPageArray(Handle SpanIndex, TArray<uint32>& PagesOut, uint32 PageOffset = 0, bool bAppend = false);
+
+	// Generates an array of contiguous ranges
+	void GetRangeArray(Handle SpanIndex, TArray<FVRamPageRange>& RangesOut, uint32 PageOffset = 0, bool bAppend = false);
+
+	// For debugging/profiling - returns the number of spans
+	uint32 GetSpanCount() const
+	{
+		return MaxSpanCount - UnusedSpanListCount - 2;
+	}
+
+	// For debugging/profiling - returns the number of outstanding allocations
+	uint32 GetAllocationCount() const
+	{
+		return AllocationCount;
+	}
+
+	uint32 GetFreePageCount() const
+	{
+		return uint32(FreePageCount);
+	}
+
+	// For debugging - returns the start offset. Note that a given allocation is not necessarily contiguous
+	uint32 GetAllocationStartPage(Handle SpanIndex) const
+	{
+		check((uint32)SpanIndex < MaxSpanCount);
+		check(PageSpans[SpanIndex].bAllocated);
+		return PageSpans[SpanIndex].StartPageIndex;
+	}
+
+	// For debugging - returns the size of this allocation in pages
+	uint32 GetAllocationPageCount(Handle SpanIndex) const
+	{
+		check((uint32)SpanIndex < MaxSpanCount);
+		check(PageSpans[SpanIndex].bAllocated);
+		uint32 Count = 0;
+		for (; SpanIndex != -1; SpanIndex = PageSpans[SpanIndex].NextSpanIndex)
+		{
+			Count += PageSpans[SpanIndex].Count;
+		}
+		return Count;
+	}
+
+	int32 GetMaxSpanCount() const
+	{
+		return MaxSpanCount;
+	}
+
+private:
+	void Init();
+
+	// Allocates pages
+	Handle AllocPagesInternal(int32 Count, uint32& NumPagesAllocated, bool bAllowPartialAlloc, const TCHAR* DebugName);
+
+	// Splits a span into two, so that the original span has PageCount pages and the new span contains the remaining ones
+	void SplitSpan(int32 SpanIndex, int32 PageCount);
+
+	// Merges two spans. They must be adjacent and in the same list
+	void MergeSpans(int32 SpanIndex0, int32 SpanIndex1, const bool bKeepSpan1);
+
+	// Inserts a span after an existing span. The span to insert must be unlinked
+	void InsertAfter(int32 InsertPosition, int32 InsertSpanIndex);
+
+	// Inserts a span after an existing span. The span to insert must be unlinked
+	void InsertBefore(int32 InsertPosition, int32 InsertSpanIndex);
+
+	// Removes a span from its list, reconnecting neighbouring list elements
+	void Unlink(int32 SpanIndex);
+
+	// Allocates an unused span from the pool
+	int AllocSpan()
+	{
+		check(UnusedSpanListCount > 0);
+		int32 SpanIndex = UnusedSpanList[UnusedSpanListCount - 1];
+		UnusedSpanListCount--;
+		return SpanIndex;
+	}
+
+	// Releases a span back to the unused pool
+	void ReleaseSpan(int32 SpanIndex)
+	{
+		check(!PageSpans[SpanIndex].IsLinked());
+		UnusedSpanList[UnusedSpanListCount] = SpanIndex;
+		UnusedSpanListCount++;
+		check(UnusedSpanListCount <= (int32)PageCount);
+	}
+
+	// Merges a span with existing neighbours in the free list if they exist
+	bool MergeFreeSpanIfPossible(int32 SpanIndex);
+
+	void Validate();
+
+	int GetFirstSpanIndex() const
+	{
+		return PageSpans[FREE_SPAN_LIST_HEAD_INDEX].NextSpanIndex;
+	}
+
+	int32 FreePageCount;
+	TArray<int32> PageToSpanStart;  // [PAGE_COUNT + 1]
+	TArray<int32> PageToSpanEnd;    // [PAGE_COUNT + 1]
+
+	TArray<FPageSpan> PageSpans;    // [MAX_SPAN_COUNT]
+	TArray<int32> UnusedSpanList;	// [MAX_SPAN_COUNT]
+	int32 UnusedSpanListCount;
+
+	const uint32 MaxSpanCount;
+	const uint32 PageCount;
+	uint32 AllocationCount;
+};
\ No newline at end of file
