Added missing copyright header.

diff --git a/Engine/Source/Runtime/RHICore/Private/RHIPageAllocator.cpp b/Engine/Source/Runtime/RHICore/Private/RHIPageAllocator.cpp
--- a/Engine/Source/Runtime/RHICore/Private/RHIPageAllocator.cpp
+++ b/Engine/Source/Runtime/RHICore/Private/RHIPageAllocator.cpp
@@ -1,3 +1,4 @@
+// Copyright Epic Games, Inc. All Rights Reserved.
 
 #include "RHIPageAllocator.h"
 
