Small fix for Tier 1 heap support on D3D12 PC - checking wrong flag value

diff --git a/Engine/Source/Runtime/RHICore/Public/RHICoreTransientResourceAllocator.h b/Engine/Source/Runtime/RHICore/Public/RHICoreTransientResourceAllocator.h
--- a/Engine/Source/Runtime/RHICore/Public/RHICoreTransientResourceAllocator.h
+++ b/Engine/Source/Runtime/RHICore/Public/RHICoreTransientResourceAllocator.h
@@ -200,9 +200,9 @@ public:
 
 	FORCEINLINE uint64 GetLastUsedGarbageCollectCycle() const { return LastUsedGarbageCollectCycle; }
 
-	FORCEINLINE bool IsAllocationSupported(uint64 Size, ERHITransientHeapFlags Flags) const
+	FORCEINLINE bool IsAllocationSupported(uint64 Size, ERHITransientHeapFlags InFlags) const
 	{
-		return Size <= Initializer.Size && EnumHasAnyFlags(Flags, Flags);
+		return Size <= Initializer.Size && EnumHasAnyFlags(Initializer.Flags, InFlags);
 	}
 
 private:
