Add GetUsedHeapStates function in RHI transient allocator to allow platform RHIs to find the total number of bytes needed in each heap.

diff --git a/Engine/Source/Runtime/RHICore/Private/RHICoreTransientResourceAllocator.cpp b/Engine/Source/Runtime/RHICore/Private/RHICoreTransientResourceAllocator.cpp
--- a/Engine/Source/Runtime/RHICore/Private/RHICoreTransientResourceAllocator.cpp
+++ b/Engine/Source/Runtime/RHICore/Private/RHICoreTransientResourceAllocator.cpp
@@ -132,17 +132,18 @@ FRHITransientHeap* FRHITransientResourceSystem::AcquireHeap(uint64 FirstAllocati
 	return CreateHeap(HeapInitializer);
 }
 
-void FRHITransientResourceSystem::ForfeitHeaps(TArrayView<FRHITransientHeap* const> InForfeitedHeaps)
+void FRHITransientResourceSystem::ForfeitHeaps(TConstArrayView<FRHITransientHeapState> InForfeitedHeaps)
 {
 	FScopeLock Lock(&HeapCriticalSection);
 
-	for (FRHITransientHeap* Heap : InForfeitedHeaps)
+	Heaps.Reserve(InForfeitedHeaps.Num());
+	for (FRHITransientHeapState const& State : InForfeitedHeaps)
 	{
-		Heap->ForfeitResources();
-		Heap->LastUsedGarbageCollectCycle = GarbageCollectCycle;
-	}
+		State.Heap->ForfeitResources();
+		State.Heap->LastUsedGarbageCollectCycle = GarbageCollectCycle;
 
-	Heaps.Append(InForfeitedHeaps.GetData(), InForfeitedHeaps.Num());
+		Heaps.Add(State.Heap);
+	}
 
 	Algo::Sort(Heaps, [](const FRHITransientHeap* LHS, const FRHITransientHeap* RHS)
 	{
@@ -554,11 +555,11 @@ FRHITransientTexture* FRHITransientResourceAllocator::CreateTexture(
 
 	const FRHITransientHeapAllocation Allocation = Allocate(TextureStats, TextureSize, TextureAlignment, TextureHeapFlags);
 
-	FRHITransientHeap* Heap = Heaps[Allocation.HeapIndex];
+	FRHITransientHeapState& HeapState = Heaps[Allocation.HeapIndex];
 
-	FRHITransientTexture* TransientTexture = Heap->AcquireTexture(CreateInfo, Allocation.Offset, [&] (uint64 Hash)
+	FRHITransientTexture* TransientTexture = HeapState.Heap->AcquireTexture(CreateInfo, Allocation.Offset, [&] (uint64 Hash)
 	{
-		const FResourceInitializer ResourceInitializer(*Heap, Allocation, Hash);
+		const FResourceInitializer ResourceInitializer(*HeapState.Heap, Allocation, Hash);
 		return CreateTextureFunction(ResourceInitializer);
 	});
 
@@ -586,11 +587,11 @@ FRHITransientBuffer* FRHITransientResourceAllocator::CreateBuffer(
 
 	const FRHITransientHeapAllocation Allocation = Allocate(BufferStats, BufferSize, BufferAlignment, ERHITransientHeapFlags::AllowBuffers);
 
-	FRHITransientHeap* Heap = Heaps[Allocation.HeapIndex];
+	FRHITransientHeapState& HeapState = Heaps[Allocation.HeapIndex];
 
-	FRHITransientBuffer* TransientBuffer = Heap->AcquireBuffer(CreateInfo, Allocation.Offset, [&](uint64 Hash)
+	FRHITransientBuffer* TransientBuffer = HeapState.Heap->AcquireBuffer(CreateInfo, Allocation.Offset, [&](uint64 Hash)
 	{
-		const FResourceInitializer ResourceInitializer(*Heap, Allocation, Hash);
+		const FResourceInitializer ResourceInitializer(*HeapState.Heap, Allocation, Hash);
 		return CreateBufferFunction(ResourceInitializer);
 	});
 
@@ -649,15 +650,19 @@ FRHITransientHeapAllocation FRHITransientResourceAllocator::Allocate(FMemoryStat
 	{
 		const uint32 HeapIndex = Heaps.Num();
 
-		FRHITransientHeap* Heap = ParentSystem.AcquireHeap(Size, ResourceHeapFlags);
-		Heaps.Emplace(Heap);
+		FRHITransientHeapState State = {};
+		State.Heap = ParentSystem.AcquireHeap(Size, ResourceHeapFlags);
+		Heaps.Emplace(State);
 
-		FRHITransientHeapAllocator& HeapAllocator = HeapAllocators.Emplace_GetRef(Heap->GetInitializer(), HeapIndex);
+		FRHITransientHeapAllocator& HeapAllocator = HeapAllocators.Emplace_GetRef(State.Heap->GetInitializer(), HeapIndex);
 		Allocation = HeapAllocator.Allocate(Size, Alignment);
 
 		check(Allocation.IsValid());
 	}
 
+	FRHITransientHeapState& State = Heaps[Allocation.HeapIndex];
+	State.CommitSize = FMath::Max(State.CommitSize, Allocation.Offset + Allocation.Size);
+
 	StatsToUpdate.AllocationCount++;
 	StatsToUpdate.TotalSize += Allocation.Size;
 	StatsToUpdate.CurrentSizeWithAliasing += Allocation.Size;
diff --git a/Engine/Source/Runtime/RHICore/Public/RHICoreTransientResourceAllocator.h b/Engine/Source/Runtime/RHICore/Public/RHICoreTransientResourceAllocator.h
--- a/Engine/Source/Runtime/RHICore/Public/RHICoreTransientResourceAllocator.h
+++ b/Engine/Source/Runtime/RHICore/Public/RHICoreTransientResourceAllocator.h
@@ -256,6 +256,15 @@ struct RHICORE_API FRHITransientResourceSystemInitializer
 	bool bSupportsAllHeapFlags = true;
 };
 
+struct FRHITransientHeapState
+{
+	FRHITransientHeap* Heap;
+
+	// The total number of bytes used by an instance of
+	// FRHITransientResourceAllocator for the given heap.
+	uint64 CommitSize = 0;
+};
+
 /** The RHI transient resource system is a base class for the platform implementation. It has a persistent lifetime
  *  and contains a cache of transient heaps. The transient allocator acquires heaps from the system and forfeits them
  *  at the end of its lifetime. Garbage collection of heaps is done using an internal counter that increments with
@@ -308,7 +317,7 @@ private:
 	FRHITransientHeap* AcquireHeap(uint64 FirstAllocationSize, ERHITransientHeapFlags FirstAllocationHeapFlags);
 
 	// Called by the transient allocator to forfeit all acquired heaps back to the cache.
-	void ForfeitHeaps(TArrayView<FRHITransientHeap* const> Heaps);
+	void ForfeitHeaps(TConstArrayView<FRHITransientHeapState> Heaps);
 
 	//////////////////////////////////////////////////////////////////////////
 	//! Platform API
@@ -566,6 +575,12 @@ public:
 	// Called to signify all allocations have completed. Forfeits all resources / heaps back to the parent system.
 	void Freeze(FRHICommandListImmediate& RHICmdList);
 
+	// Returns the array of heaps used by this allocator, including the required commit size for each.
+	inline TConstArrayView<FRHITransientHeapState> GetUsedHeapStates() const
+	{
+		return Heaps;
+	}
+
 private:
 	struct FMemoryStats : FRHITransientMemoryStats
 	{
@@ -581,7 +596,7 @@ private:
 	FRHITransientResourceSystem& ParentSystem;
 
 	// Tracks state on the rendering thread; must be cleared before the destructor.
-	TArray<FRHITransientHeap*, TInlineAllocator<4, TMemStackAllocator<>>> Heaps;
+	TArray<FRHITransientHeapState, TInlineAllocator<4, TMemStackAllocator<>>> Heaps;
 	TArray<FRHITransientHeapAllocator, TInlineAllocator<4, TMemStackAllocator<>>> HeapAllocators;
 	TArray<FRHITransientHeapAllocation, TInlineAllocator<4, TMemStackAllocator<>>> HeapAllocations;
 
@@ -592,7 +607,4 @@ private:
 
 	FMemoryStats TextureStats;
 	FMemoryStats BufferStats;
-
-	uint64 AllocatedTextureSize{};
-	uint64 AllocatedBufferSize{};
 };
\ No newline at end of file
