Fix overlap check for transient resources

diff --git a/Engine/Source/Runtime/RHICore/Private/RHICoreTransientResourceAllocator.cpp b/Engine/Source/Runtime/RHICore/Private/RHICoreTransientResourceAllocator.cpp
--- a/Engine/Source/Runtime/RHICore/Private/RHICoreTransientResourceAllocator.cpp
+++ b/Engine/Source/Runtime/RHICore/Private/RHICoreTransientResourceAllocator.cpp
@@ -455,7 +455,7 @@ void FRHITransientResourceOverlapTracker::Track(FRHITransientResource* Transient
 		FResourceRange& ResourceRangeOld = ResourceRanges[Index];
 
 		// If the old range starts later in the heap and doesn't overlap, the sort invariant guarantees no future range will overlap.
-		if (ResourceRangeOld.OffsetMin >= ResourceRangeNew.OffsetMin)
+		if (ResourceRangeOld.OffsetMin >= ResourceRangeNew.OffsetMax)
 		{
 			break;
 		}
