Fix incorrect placed resource alignments due to not taking the heap base address into account

diff --git a/Engine/Source/Runtime/D3D12RHI/Private/D3D12Resources.cpp b/Engine/Source/Runtime/D3D12RHI/Private/D3D12Resources.cpp
--- a/Engine/Source/Runtime/D3D12RHI/Private/D3D12Resources.cpp
+++ b/Engine/Source/Runtime/D3D12RHI/Private/D3D12Resources.cpp
@@ -429,7 +429,7 @@ FD3D12Heap::~FD3D12Heap()
 	Heap.SafeRelease();
 }
 
-void FD3D12Heap::SetHeap(ID3D12Heap* HeapIn, const TCHAR* const InName, bool bInTrack)
+void FD3D12Heap::SetHeap(ID3D12Heap* HeapIn, const TCHAR* const InName, bool bInTrack, bool bForceGetGPUAddress)
 {
 	*Heap.GetInitReference() = HeapIn; 
 	bTrack = bInTrack;
@@ -441,7 +441,7 @@ void FD3D12Heap::SetHeap(ID3D12Heap* HeapIn, const TCHAR* const InName, bool bIn
 	// Create a buffer placed resource on the heap to extract the gpu virtual address
 	// if we are tracking all allocations
 	FD3D12Adapter* Adapter = GetParentDevice()->GetParentAdapter();	
-	if (Adapter->IsTrackingAllAllocations() && HeapDesc.Properties.Type == D3D12_HEAP_TYPE_DEFAULT)
+	if ((bForceGetGPUAddress || Adapter->IsTrackingAllAllocations()) && HeapDesc.Properties.Type == D3D12_HEAP_TYPE_DEFAULT)
 	{
 		uint64 HeapSize = HeapDesc.SizeInBytes;
 		TRefCountPtr<ID3D12Resource> TempResource;
diff --git a/Engine/Source/Runtime/D3D12RHI/Private/D3D12TransientResourceAllocator.cpp b/Engine/Source/Runtime/D3D12RHI/Private/D3D12TransientResourceAllocator.cpp
--- a/Engine/Source/Runtime/D3D12RHI/Private/D3D12TransientResourceAllocator.cpp
+++ b/Engine/Source/Runtime/D3D12RHI/Private/D3D12TransientResourceAllocator.cpp
@@ -84,8 +84,10 @@ FD3D12TransientHeap::FD3D12TransientHeap(const FRHITransientHeapInitializer& Ini
 	}
 
 	Heap = new FD3D12Heap(Device, VisibleNodeMask);
-	Heap->SetHeap(D3DHeap, TEXT("TransientResourceAllocator Backing Heap"));
+	Heap->SetHeap(D3DHeap, TEXT("TransientResourceAllocator Backing Heap"), true, true);
 	Heap->BeginTrackingResidency(Desc.SizeInBytes);
+
+	BaseGPUVirtualAddress = Heap->GetGPUVirtualAddress();
 }
 
 FD3D12TransientHeap::~FD3D12TransientHeap()
diff --git a/Engine/Source/Runtime/D3D12RHI/Public/D3D12Resources.h b/Engine/Source/Runtime/D3D12RHI/Public/D3D12Resources.h
--- a/Engine/Source/Runtime/D3D12RHI/Public/D3D12Resources.h
+++ b/Engine/Source/Runtime/D3D12RHI/Public/D3D12Resources.h
@@ -100,7 +100,7 @@ public:
 	~FD3D12Heap();
 
 	inline ID3D12Heap* GetHeap() const { return Heap.GetReference(); }
-	void SetHeap(ID3D12Heap* HeapIn, const TCHAR* const InName, bool bTrack = true);
+	void SetHeap(ID3D12Heap* HeapIn, const TCHAR* const InName, bool bTrack = true, bool bForceGetGPUAddress = false);
 
 	void UpdateResidency(FD3D12CommandListHandle& CommandList);
 	void BeginTrackingResidency(uint64 Size);
diff --git a/Engine/Source/Runtime/RHICore/Private/RHICoreTransientResourceAllocator.cpp b/Engine/Source/Runtime/RHICore/Private/RHICoreTransientResourceAllocator.cpp
--- a/Engine/Source/Runtime/RHICore/Private/RHICoreTransientResourceAllocator.cpp
+++ b/Engine/Source/Runtime/RHICore/Private/RHICoreTransientResourceAllocator.cpp
@@ -226,9 +226,10 @@ void FRHITransientResourceSystem::UpdateStats()
 
 //////////////////////////////////////////////////////////////////////////
 
-FRHITransientHeapAllocator::FRHITransientHeapAllocator(const FRHITransientHeapInitializer& InInitializer, uint32 InHeapIndex)
+FRHITransientHeapAllocator::FRHITransientHeapAllocator(const FRHITransientHeapInitializer& InInitializer, uint32 InHeapIndex, uint64 InHeapBaseGPUAddress)
 	: Initializer(InInitializer)
 	, HeapIndex(InHeapIndex)
+	, HeapBaseGPUAddress(InHeapBaseGPUAddress)
 {
 	HeadHandle = CreateRange();
 	InsertRange(HeadHandle, 0, Initializer.Size);
@@ -392,7 +393,7 @@ FRHITransientHeapAllocator::FFindResult FRHITransientHeapAllocator::FindFreeRang
 		FRange& Range = Ranges[Handle];
 
 		// Due to alignment we may have to shift the offset and expand the size accordingly.
-		const uint64 AlignmentPad = Align(Range.Offset, Alignment) - Range.Offset;
+		const uint64 AlignmentPad = Align(HeapBaseGPUAddress + Range.Offset, Alignment) - HeapBaseGPUAddress - Range.Offset;
 		const uint64 RequiredSize = Size + AlignmentPad;
 
 		if (RequiredSize <= Range.Size)
@@ -654,7 +655,7 @@ FRHITransientHeapAllocation FRHITransientResourceAllocator::Allocate(FMemoryStat
 		State.Heap = ParentSystem.AcquireHeap(Size, ResourceHeapFlags);
 		Heaps.Emplace(State);
 
-		FRHITransientHeapAllocator& HeapAllocator = HeapAllocators.Emplace_GetRef(State.Heap->GetInitializer(), HeapIndex);
+		FRHITransientHeapAllocator& HeapAllocator = HeapAllocators.Emplace_GetRef(State.Heap->GetInitializer(), HeapIndex, State.Heap->GetBaseGPUVirtualAddress());
 		Allocation = HeapAllocator.Allocate(Size, Alignment);
 
 		check(Allocation.IsValid());
diff --git a/Engine/Source/Runtime/RHICore/Public/RHICoreTransientResourceAllocator.h b/Engine/Source/Runtime/RHICore/Public/RHICoreTransientResourceAllocator.h
--- a/Engine/Source/Runtime/RHICore/Public/RHICoreTransientResourceAllocator.h
+++ b/Engine/Source/Runtime/RHICore/Public/RHICoreTransientResourceAllocator.h
@@ -196,6 +196,8 @@ public:
 
 	FORCEINLINE const FRHITransientHeapInitializer& GetInitializer() const { return Initializer; }
 
+	FORCEINLINE uint64 GetBaseGPUVirtualAddress() const { return BaseGPUVirtualAddress; }
+
 	FORCEINLINE uint64 GetCapacity() const { return Initializer.Size; }
 
 	FORCEINLINE uint64 GetLastUsedGarbageCollectCycle() const { return LastUsedGarbageCollectCycle; }
@@ -205,6 +207,9 @@ public:
 		return Size <= Initializer.Size && EnumHasAnyFlags(Initializer.Flags, InFlags);
 	}
 
+protected:
+	uint64 BaseGPUVirtualAddress = 0;
+
 private:
 	template <typename CreateFunctionType>
 	FRHITransientTexture* AcquireTexture(const FRHITextureCreateInfo& CreateInfo, uint64 HeapOffset, CreateFunctionType CreateFunction);
@@ -394,7 +399,7 @@ struct FRHITransientHeapAllocation
 class RHICORE_API FRHITransientHeapAllocator final
 {
 public:
-	FRHITransientHeapAllocator(const FRHITransientHeapInitializer& Initializer, uint32 HeapIndex);
+	FRHITransientHeapAllocator(const FRHITransientHeapInitializer& Initializer, uint32 HeapIndex, uint64 InHeapBaseGPUAddress = 0);
 	~FRHITransientHeapAllocator();
 
 	FRHITransientHeapAllocation Allocate(uint64 Size, uint32 Alignment);
@@ -493,6 +498,8 @@ private:
 	uint32 AllocationCount{};
 	uint32 HeapIndex{};
 
+	uint64 HeapBaseGPUAddress;
+
 	FRangeHandle HeadHandle = InvalidRangeHandle;
 	TArray<FRangeHandle, TInlineAllocator<4, TMemStackAllocator<>>> RangeFreeList;
 	TArray<FRange, TInlineAllocator<4, TMemStackAllocator<>>> Ranges;
