Implemented FRHITransientResourceAllocator::Flush to fix RDG drain GPU segfault.

diff --git a/Engine/Source/Runtime/RHICore/Private/RHICoreTransientResourceAllocator.cpp b/Engine/Source/Runtime/RHICore/Private/RHICoreTransientResourceAllocator.cpp
--- a/Engine/Source/Runtime/RHICore/Private/RHICoreTransientResourceAllocator.cpp
+++ b/Engine/Source/Runtime/RHICore/Private/RHICoreTransientResourceAllocator.cpp
@@ -626,6 +626,14 @@ void FRHITransientResourceAllocator::Freeze(FRHICommandListImmediate& RHICmdList
 	});
 }
 
+void FRHITransientResourceAllocator::Flush()
+{
+	for (FRHITransientHeapState& State : Heaps)
+	{
+		State.LastCommitSize = State.CommitSize;
+	}
+}
+
 FRHITransientHeapAllocation FRHITransientResourceAllocator::Allocate(FMemoryStats& StatsToUpdate, uint64 Size, uint32 Alignment, ERHITransientHeapFlags ResourceHeapFlags)
 {
 	FRHITransientHeapAllocation Allocation;
diff --git a/Engine/Source/Runtime/RHICore/Public/RHICoreTransientResourceAllocator.h b/Engine/Source/Runtime/RHICore/Public/RHICoreTransientResourceAllocator.h
--- a/Engine/Source/Runtime/RHICore/Public/RHICoreTransientResourceAllocator.h
+++ b/Engine/Source/Runtime/RHICore/Public/RHICoreTransientResourceAllocator.h
@@ -268,6 +268,12 @@ struct FRHITransientHeapState
 	// The total number of bytes used by an instance of
 	// FRHITransientResourceAllocator for the given heap.
 	uint64 CommitSize = 0;
+
+	// The commit size since the last flush operation.
+	uint64 LastCommitSize = 0;
+
+	// Returns whether a commit is required for the heap since the last flush. This is reset with each Flush() call on the allocator.
+	bool IsCommitRequired() const { return CommitSize > LastCommitSize; }
 };
 
 /** The RHI transient resource system is a base class for the platform implementation. It has a persistent lifetime
@@ -582,6 +588,9 @@ public:
 	// Called to signify all allocations have completed. Forfeits all resources / heaps back to the parent system.
 	void Freeze(FRHICommandListImmediate& RHICmdList);
 
+	// Called to flush any active allocations prior to Freeze.
+	void Flush();
+
 	// Returns the array of heaps used by this allocator, including the required commit size for each.
 	inline TConstArrayView<FRHITransientHeapState> GetUsedHeapStates() const
 	{
