Optimized transient allocator. Reduces cost on AncientGame by ~0.05ms.

diff --git a/Engine/Source/Runtime/RHICore/Private/RHICoreTransientResourceAllocator.cpp b/Engine/Source/Runtime/RHICore/Private/RHICoreTransientResourceAllocator.cpp
--- a/Engine/Source/Runtime/RHICore/Private/RHICoreTransientResourceAllocator.cpp
+++ b/Engine/Source/Runtime/RHICore/Private/RHICoreTransientResourceAllocator.cpp
@@ -45,20 +45,6 @@ DECLARE_DWORD_ACCUMULATOR_STAT(TEXT("Textures"), STAT_RHITransientTextures, STAT
 DECLARE_DWORD_ACCUMULATOR_STAT(TEXT("Buffers"), STAT_RHITransientBuffers, STATGROUP_RHITransientMemory);
 DECLARE_DWORD_ACCUMULATOR_STAT(TEXT("Heaps"), STAT_RHITransientHeaps, STATGROUP_RHITransientMemory);
 
-inline uint64 ComputeHash(const FRHITextureCreateInfo& InCreateInfo, uint64 HeapOffset)
-{
-	// Make sure all padding is removed.
-	FRHITextureCreateInfo NewInfo;
-	FPlatformMemory::Memzero(&NewInfo, sizeof(FRHITextureCreateInfo));
-	NewInfo = InCreateInfo;
-	return CityHash64WithSeed((const char*)&NewInfo, sizeof(FRHITextureCreateInfo), HeapOffset);
-}
-
-inline uint64 ComputeHash(const FRHIBufferCreateInfo& InCreateInfo, uint64 HeapOffset)
-{
-	return CityHash64WithSeed((const char*)&InCreateInfo, sizeof(FRHIBufferCreateInfo), HeapOffset);
-}
-
 //////////////////////////////////////////////////////////////////////////
 
 template <typename CreateFunctionType>
diff --git a/Engine/Source/Runtime/RHICore/Public/RHICoreTransientResourceAllocator.h b/Engine/Source/Runtime/RHICore/Public/RHICoreTransientResourceAllocator.h
--- a/Engine/Source/Runtime/RHICore/Public/RHICoreTransientResourceAllocator.h
+++ b/Engine/Source/Runtime/RHICore/Public/RHICoreTransientResourceAllocator.h
@@ -14,6 +14,20 @@ class FRHITransientResourceAllocator;
 	#define IF_RHICORE_TRANSIENT_ALLOCATOR_DEBUG(Op)
 #endif
 
+inline uint64 ComputeHash(const FRHITextureCreateInfo& InCreateInfo, uint64 HeapOffset)
+{
+	// Make sure all padding is removed.
+	FRHITextureCreateInfo NewInfo;
+	FPlatformMemory::Memzero(&NewInfo, sizeof(FRHITextureCreateInfo));
+	NewInfo = InCreateInfo;
+	return CityHash64WithSeed((const char*)&NewInfo, sizeof(FRHITextureCreateInfo), HeapOffset);
+}
+
+inline uint64 ComputeHash(const FRHIBufferCreateInfo& InCreateInfo, uint64 HeapOffset)
+{
+	return CityHash64WithSeed((const char*)&InCreateInfo, sizeof(FRHIBufferCreateInfo), HeapOffset);
+}
+
 struct FRHITransientMemoryStats
 {
 	FRHITransientMemoryStats() = default;
