Fixed bug in transient allocator overlap tracker causing large resources to be lost as aliases.

diff --git a/Engine/Source/Runtime/RHICore/Private/RHICoreTransientResourceAllocator.cpp b/Engine/Source/Runtime/RHICore/Private/RHICoreTransientResourceAllocator.cpp
--- a/Engine/Source/Runtime/RHICore/Private/RHICoreTransientResourceAllocator.cpp
+++ b/Engine/Source/Runtime/RHICore/Private/RHICoreTransientResourceAllocator.cpp
@@ -465,6 +465,7 @@ void FRHITransientResourceOverlapTracker::Track(FRHITransientResource* Transient
 		else
 		{
 			bool bResizedOld = false;
+			FResourceRange ResourceRangeOldCopy = ResourceRangeOld;
 
 			// 1) New:    ********
 			//            |||        ->
@@ -475,7 +476,7 @@ void FRHITransientResourceOverlapTracker::Track(FRHITransientResource* Transient
 				bResizedOld = true;
 			}
 
-			if (ResourceRangeOld.OffsetMax > ResourceRangeNew.OffsetMax)
+			if (ResourceRangeOldCopy.OffsetMax > ResourceRangeNew.OffsetMax)
 			{
 				// 2) New:    ********
 				//                |||      ->
@@ -491,7 +492,7 @@ void FRHITransientResourceOverlapTracker::Track(FRHITransientResource* Transient
 				else
 				{
 					// Lower bound has been resized already; add an upper bound.
-					FResourceRange ResourceRangeOldUpper = ResourceRangeOld;
+					FResourceRange ResourceRangeOldUpper = ResourceRangeOldCopy;
 					ResourceRangeOldUpper.OffsetMin = ResourceRangeNew.OffsetMax;
 					Index = ResourceRanges.Insert(ResourceRangeOldUpper, Index + 1);
 				}
