Adding a RHI Transient Resource LLM tag and making sure D3D12 uses it.

diff --git a/Engine/Source/Runtime/D3D12RHI/Private/D3D12TransientResourceAllocator.cpp b/Engine/Source/Runtime/D3D12RHI/Private/D3D12TransientResourceAllocator.cpp
--- a/Engine/Source/Runtime/D3D12RHI/Private/D3D12TransientResourceAllocator.cpp
+++ b/Engine/Source/Runtime/D3D12RHI/Private/D3D12TransientResourceAllocator.cpp
@@ -68,8 +68,6 @@ FD3D12TransientHeap::FD3D12TransientHeap(const FRHITransientHeapInitializer& Ini
 
 		LLM_PLATFORM_SCOPE(ELLMTag::GraphicsPlatform);
 
-		// We are tracking allocations ourselves, so don't let XMemAlloc track these as well
-		LLM_SCOPED_PAUSE_TRACKING_FOR_TRACKER(ELLMTracker::Default, ELLMAllocType::System);
 		VERIFYD3D12RESULT(D3DDevice->CreateHeap(&Desc, IID_PPV_ARGS(&D3DHeap)));
 
 #if PLATFORM_WINDOWS
@@ -95,8 +93,6 @@ FD3D12TransientHeap::FD3D12TransientHeap(const FRHITransientHeapInitializer& Ini
 
 FD3D12TransientHeap::~FD3D12TransientHeap()
 {
-	LLM_SCOPED_PAUSE_TRACKING_FOR_TRACKER(ELLMTracker::Default, ELLMAllocType::System);
-
 	if (Heap)
 	{
 		D3D12_HEAP_DESC Desc = Heap->GetHeapDesc();
diff --git a/Engine/Source/Runtime/RHICore/Private/RHICoreTransientResourceAllocator.cpp b/Engine/Source/Runtime/RHICore/Private/RHICoreTransientResourceAllocator.cpp
--- a/Engine/Source/Runtime/RHICore/Private/RHICoreTransientResourceAllocator.cpp
+++ b/Engine/Source/Runtime/RHICore/Private/RHICoreTransientResourceAllocator.cpp
@@ -1,6 +1,8 @@
 // Copyright Epic Games, Inc. All Rights Reserved.
 
 #include "RHICoreTransientResourceAllocator.h"
+#include "HAL/LowLevelMemTracker.h"
+#include "HAL/LowLevelMemStats.h"
 
 static int32 GRHITransientAllocatorMinimumHeapSize = 128;
 static FAutoConsoleVariableRef CVarGRHITransientAllocatorMinimumHeapSize(
@@ -45,6 +47,9 @@ DECLARE_DWORD_ACCUMULATOR_STAT(TEXT("Textures"), STAT_RHITransientTextures, STAT
 DECLARE_DWORD_ACCUMULATOR_STAT(TEXT("Buffers"), STAT_RHITransientBuffers, STATGROUP_RHITransientMemory);
 DECLARE_DWORD_ACCUMULATOR_STAT(TEXT("Heaps"), STAT_RHITransientHeaps, STATGROUP_RHITransientMemory);
 
+DECLARE_LLM_MEMORY_STAT(TEXT("RHI Transient Resources"), STAT_RHITransientResourcesLLM, STATGROUP_LLMFULL);
+LLM_DEFINE_TAG(RHITransientResources, NAME_None, NAME_None, GET_STATFNAME(STAT_RHITransientResourcesLLM), GET_STATFNAME(STAT_EngineSummaryLLM));
+
 //////////////////////////////////////////////////////////////////////////
 
 template <typename CreateFunctionType>
@@ -115,6 +120,7 @@ FRHITransientHeap* FRHITransientResourceSystem::AcquireHeap(uint64 FirstAllocati
 	HeapInitializer.TextureCacheSize = Initializer.TextureCacheSize;
 	HeapInitializer.BufferCacheSize = Initializer.BufferCacheSize;
 
+	LLM_SCOPE_BYTAG(RHITransientResources);
 	return CreateHeap(HeapInitializer);
 }
 
