Take resource alignment into account when creating a transient heap.

The default alignment for transient heaps is 64k, but resources can have larger alignment requirements. We need to take this into account when selecting or creating heaps, otherwise the allocation can fail even though the heap seems to have enough room. Also, the alignment value wasn't actually used when the D3D heap object was created, we always passed 0 to ID3D12Device::CreateHeap.

#rnx

diff --git a/Engine/Source/Runtime/D3D12RHI/Private/D3D12TransientResourceAllocator.cpp b/Engine/Source/Runtime/D3D12RHI/Private/D3D12TransientResourceAllocator.cpp
--- a/Engine/Source/Runtime/D3D12RHI/Private/D3D12TransientResourceAllocator.cpp
+++ b/Engine/Source/Runtime/D3D12RHI/Private/D3D12TransientResourceAllocator.cpp
@@ -54,7 +54,7 @@ FD3D12TransientHeap::FD3D12TransientHeap(const FRHITransientHeapInitializer& Ini
 	D3D12_HEAP_DESC Desc = {};
 	Desc.SizeInBytes = Initializer.Size;
 	Desc.Properties = HeapProperties;
-	Desc.Alignment = 0;
+	Desc.Alignment = Initializer.Alignment;
 	Desc.Flags = HeapFlags;
 
 	if (Adapter->IsHeapNotZeroedSupported())
diff --git a/Engine/Source/Runtime/RHICore/Private/RHICoreTransientResourceAllocator.cpp b/Engine/Source/Runtime/RHICore/Private/RHICoreTransientResourceAllocator.cpp
--- a/Engine/Source/Runtime/RHICore/Private/RHICoreTransientResourceAllocator.cpp
+++ b/Engine/Source/Runtime/RHICore/Private/RHICoreTransientResourceAllocator.cpp
@@ -96,7 +96,7 @@ FRHITransientResourceSystem::~FRHITransientResourceSystem()
 	ReleaseHeaps();
 }
 
-FRHITransientHeap* FRHITransientResourceSystem::AcquireHeap(uint64 FirstAllocationSize, ERHITransientHeapFlags FirstAllocationHeapFlags)
+FRHITransientHeap* FRHITransientResourceSystem::AcquireHeap(uint64 FirstAllocationSize, uint32 FirstAllocationAlignment, ERHITransientHeapFlags FirstAllocationHeapFlags)
 {
 	{
 		FScopeLock Lock(&HeapCriticalSection);
@@ -105,7 +105,7 @@ FRHITransientHeap* FRHITransientResourceSystem::AcquireHeap(uint64 FirstAllocati
 		{
 			FRHITransientHeap* Heap = Heaps[HeapIndex];
 
-			if (Heap->IsAllocationSupported(FirstAllocationSize, FirstAllocationHeapFlags))
+			if (Heap->IsAllocationSupported(FirstAllocationSize, FirstAllocationAlignment, FirstAllocationHeapFlags))
 			{
 				Heaps.RemoveAt(HeapIndex);
 				return Heap;
@@ -115,7 +115,7 @@ FRHITransientHeap* FRHITransientResourceSystem::AcquireHeap(uint64 FirstAllocati
 
 	FRHITransientHeapInitializer HeapInitializer;
 	HeapInitializer.Size = GetHeapSize(FirstAllocationSize);
-	HeapInitializer.Alignment = GetHeapAlignment();
+	HeapInitializer.Alignment = GetHeapAlignment(FirstAllocationAlignment);
 	HeapInitializer.Flags = (Initializer.bSupportsAllHeapFlags ? ERHITransientHeapFlags::AllowAll : FirstAllocationHeapFlags);
 	HeapInitializer.TextureCacheSize = Initializer.TextureCacheSize;
 	HeapInitializer.BufferCacheSize = Initializer.BufferCacheSize;
@@ -657,7 +657,7 @@ FRHITransientHeapAllocation FRHITransientResourceAllocator::Allocate(FMemoryStat
 		const uint32 HeapIndex = Heaps.Num();
 
 		FRHITransientHeapState State = {};
-		State.Heap = ParentSystem.AcquireHeap(Size, ResourceHeapFlags);
+		State.Heap = ParentSystem.AcquireHeap(Size, Alignment, ResourceHeapFlags);
 		Heaps.Emplace(State);
 
 		FRHITransientHeapAllocator& HeapAllocator = HeapAllocators.Emplace_GetRef(State.Heap->GetInitializer(), HeapIndex, State.Heap->GetBaseGPUVirtualAddress());
diff --git a/Engine/Source/Runtime/RHICore/Public/RHICoreTransientResourceAllocator.h b/Engine/Source/Runtime/RHICore/Public/RHICoreTransientResourceAllocator.h
--- a/Engine/Source/Runtime/RHICore/Public/RHICoreTransientResourceAllocator.h
+++ b/Engine/Source/Runtime/RHICore/Public/RHICoreTransientResourceAllocator.h
@@ -216,9 +216,9 @@ public:
 
 	FORCEINLINE uint64 GetLastUsedGarbageCollectCycle() const { return LastUsedGarbageCollectCycle; }
 
-	FORCEINLINE bool IsAllocationSupported(uint64 Size, ERHITransientHeapFlags InFlags) const
+	FORCEINLINE bool IsAllocationSupported(uint64 Size, uint32 Alignment, ERHITransientHeapFlags InFlags) const
 	{
-		return Size <= Initializer.Size && EnumHasAnyFlags(Initializer.Flags, InFlags);
+		return Size <= Initializer.Size && Alignment <= Initializer.Alignment && EnumHasAnyFlags(Initializer.Flags, InFlags);
 	}
 
 protected:
@@ -322,7 +322,7 @@ public:
 
 	FORCEINLINE uint64 GetMinimumHeapSize() const { return Initializer.MinimumHeapSize; }
 	FORCEINLINE uint64 GetMaximumHeapSize() const { return Initializer.MaximumHeapSize; }
-	FORCEINLINE uint32 GetHeapAlignment() const { return Initializer.HeapAlignment; }
+	FORCEINLINE uint32 GetHeapAlignment(uint32 FirstAllocationAlignment) const { return FMath::Max(Initializer.HeapAlignment, FirstAllocationAlignment); }
 	FORCEINLINE uint32 GetHeapCount() const { return Heaps.Num(); }
 
 	FORCEINLINE FRHITransientHeap* GetHeap(uint32 Index) { return Heaps[Index]; }
@@ -339,7 +339,7 @@ protected:
 
 private:
 	// Called by the transient allocator to acquire a heap from the cache.
-	FRHITransientHeap* AcquireHeap(uint64 FirstAllocationSize, ERHITransientHeapFlags FirstAllocationHeapFlags);
+	FRHITransientHeap* AcquireHeap(uint64 FirstAllocationSize, uint32 FirstAllocationAlignment, ERHITransientHeapFlags FirstAllocationHeapFlags);
 
 	// Called by the transient allocator to forfeit all acquired heaps back to the cache.
 	void ForfeitHeaps(TConstArrayView<FRHITransientHeapState> Heaps);
