Increase garbage collect latency for deletion of transient heaps to handle frame peaks & stutters - otherwise heaps were destroyed and immediatly recreated before being fully destroyed causing memory peaks of double the required memory in Frosty

diff --git a/Engine/Source/Runtime/RHICore/Private/RHICoreTransientResourceAllocator.cpp b/Engine/Source/Runtime/RHICore/Private/RHICoreTransientResourceAllocator.cpp
--- a/Engine/Source/Runtime/RHICore/Private/RHICoreTransientResourceAllocator.cpp
+++ b/Engine/Source/Runtime/RHICore/Private/RHICoreTransientResourceAllocator.cpp
@@ -32,6 +32,13 @@ static FAutoConsoleVariableRef CVarRHITransientAllocatorTextureCacheSize(
 	TEXT("The maximum number of RHI textures to cache on each heap before garbage collecting."),
 	ECVF_ReadOnly);
 
+static int32 GRHITransientAllocatorHeapGarbageCollectLatency = 20;
+static FAutoConsoleVariableRef CVarRHITransientAllocatorHeapGarbageCollectLatency(
+	TEXT("RHI.TransientAllocator.HeapGarbageCollectLatency"),
+	GRHITransientAllocatorHeapGarbageCollectLatency,
+	TEXT("Amount of update cycles before a heap is deleted when not used."),
+	ECVF_ReadOnly);
+
 DECLARE_STATS_GROUP(TEXT("RHI: Transient Memory"), STATGROUP_RHITransientMemory, STATCAT_Advanced);
 
 DECLARE_MEMORY_STAT(TEXT("Memory Allocated"), STAT_RHITransientMemoryAllocated, STATGROUP_RHITransientMemory);
@@ -88,6 +95,7 @@ FRHITransientResourceSystemInitializer FRHITransientResourceSystemInitializer::C
 	Initializer.MaximumHeapSize = GRHITransientAllocatorMaximumHeapSize * 1024 * 1024;
 	Initializer.BufferCacheSize = GRHITransientAllocatorBufferCacheSize;
 	Initializer.TextureCacheSize = GRHITransientAllocatorTextureCacheSize;
+	Initializer.GarbageCollectLatency = GRHITransientAllocatorHeapGarbageCollectLatency;
 	return Initializer;
 }
 
@@ -694,4 +702,4 @@ void FRHITransientResourceAllocator::InitResource(FRHITransientResource* Transie
 	TransientResource->Init(Name, HeapAllocations.Emplace(Allocation), PassIndex, Stats);
 
 	HeapAllocators[Allocation.HeapIndex].TrackOverlap(TransientResource, Allocation);
-}
\ No newline at end of file
+}
diff --git a/Engine/Source/Runtime/RHICore/Public/RHICoreTransientResourceAllocator.h b/Engine/Source/Runtime/RHICore/Public/RHICoreTransientResourceAllocator.h
--- a/Engine/Source/Runtime/RHICore/Public/RHICoreTransientResourceAllocator.h
+++ b/Engine/Source/Runtime/RHICore/Public/RHICoreTransientResourceAllocator.h
@@ -250,7 +250,6 @@ struct RHICORE_API FRHITransientResourceSystemInitializer
 	// Creates a default initializer using common RHI CVars.
 	static FRHITransientResourceSystemInitializer CreateDefault();
 
-	static const uint32 kDefaultGarbageCollectLatency = 20;
 	static const uint32 kDefaultResourceCacheSize = 256;
 
 	// The minimum size to use when creating a heap. This is the default but can grow based on allocations.
@@ -263,7 +262,7 @@ struct RHICORE_API FRHITransientResourceSystemInitializer
 	uint32 HeapAlignment{};
 
 	// The latency between the completed fence value and the used fence value to invoke garbage collection of the heap.
-	uint32 GarbageCollectLatency = kDefaultGarbageCollectLatency;
+	uint32 GarbageCollectLatency{};
 
 	// Size of the texture cache before elements are evicted.
 	uint32 TextureCacheSize = kDefaultResourceCacheSize;
