- Reneable defrag on suballocated pools
- Timeslice defrag on pool by default (max 1 pool per allocator)
- Reduce max allocation size to 4MB on pool instead of 8MB (reduce fragmentation)

diff --git a/Engine/Source/Runtime/D3D12RHI/Private/D3D12Allocation.cpp b/Engine/Source/Runtime/D3D12RHI/Private/D3D12Allocation.cpp
--- a/Engine/Source/Runtime/D3D12RHI/Private/D3D12Allocation.cpp
+++ b/Engine/Source/Runtime/D3D12RHI/Private/D3D12Allocation.cpp
@@ -1338,7 +1338,7 @@ FD3D12BufferPool* FD3D12DefaultBufferAllocator::CreateBufferPool(D3D12_HEAP_TYPE
 	FRHIMemoryPool::EFreeListOrder FreeListOrder = FRHIMemoryPool::EFreeListOrder::SortBySize;
 
 	// Disable defrag if not Default memory
-	bool bDefragEnabled = (InitConfig.HeapType == D3D12_HEAP_TYPE_DEFAULT && AllocationStrategy == EResourceAllocationStrategy::kPlacedResource);
+	bool bDefragEnabled = (InitConfig.HeapType == D3D12_HEAP_TYPE_DEFAULT);
 
 #if D3D12_RHI_RAYTRACING
 	// Disable defrag on the RT Acceleration pool - #kenzo_todo
@@ -1496,6 +1496,12 @@ void FD3D12DefaultBufferAllocator::BeginFrame()
 			if (DefaultBufferPool)
 			{
 				DefaultBufferPool->Defrag(MaxCopySize, CopySize);
+
+				// break when we reach the max copy size
+				if (CopySize >= MaxCopySize)
+				{
+					break;
+				}
 			}
 		}
 	}
diff --git a/Engine/Source/Runtime/RHICore/Private/RHIPoolAllocator.cpp b/Engine/Source/Runtime/RHICore/Private/RHIPoolAllocator.cpp
--- a/Engine/Source/Runtime/RHICore/Private/RHIPoolAllocator.cpp
+++ b/Engine/Source/Runtime/RHICore/Private/RHIPoolAllocator.cpp
@@ -2,8 +2,22 @@
 
 #include "RHIPoolAllocator.h"
 
+static float GRHIPoolAllocatorDefragSizeFraction = 0.9f;
+static FAutoConsoleVariableRef CVarRHIPoolAllocatorDefragSizeFraction(
+	TEXT("RHIPoolAllocator.DefragSizeFraction"),
+	GRHIPoolAllocatorDefragSizeFraction,
+	TEXT("Skip defrag of pool if usage is more than given fraction (default: 0.9f)."),
+	ECVF_Default);
+
+static int32 GRHIPoolAllocatorDefragMaxPoolsToClear = 1;
+static FAutoConsoleVariableRef CVarRHIPoolAllocatorDefragMaxPoolsToClear(
+	TEXT("RHIPoolAllocator.DefragMaxPoolsToClear"),
+	GRHIPoolAllocatorDefragMaxPoolsToClear,
+	TEXT("Maximum amount of pools to try and clear during a single alloctor defrag call (default: 1 - negative then all pools will be tried and no timeslicing)"),
+	ECVF_Default);
+
 static int32 GRHIPoolAllocatorValidateLinkList = 0;
-static FAutoConsoleVariableRef CVarD3D12FastAllocatorMinPagesToRetain(
+static FAutoConsoleVariableRef CVarRHIPoolAllocatorValidateLinkList(
 	TEXT("RHIPoolAllocator.ValidateLinkedList"),
 	GRHIPoolAllocatorValidateLinkList,
 	TEXT("Validate all the internal linked list data of all the RHIPoolAllocators during every operation."),
@@ -369,6 +383,8 @@ void FRHIMemoryPool::Deallocate(FRHIPoolAllocationData& AllocationData)
 
 void FRHIMemoryPool::TryClear(FRHIPoolAllocator* InAllocator, uint32 InMaxCopySize, uint32& CopySize, const TArray<FRHIMemoryPool*>& InTargetPools)
 {
+	TRACE_CPUPROFILER_EVENT_SCOPE(FRHIMemoryPool::TryClear);
+
 	FRHIPoolAllocationData* BlockToMove = HeadBlock.GetPrev();
 	while (BlockToMove != &HeadBlock && CopySize < InMaxCopySize)
 	{
@@ -415,14 +431,32 @@ int32 FRHIMemoryPool::FindFreeBlock(uint32 InSizeInBytes, uint32 InAllocationAli
 		return INDEX_NONE;
 	}
 
-	// TODO: lowerbound search
+	if (FreeListOrder == EFreeListOrder::SortBySize)
+	{
+		// check first index as fast path
+		if (FreeBlocks[0]->GetSize() >= AlignedSize)
+		{
+			return 0;
+		}
 
-	// Check all the free blocks (sorted by size, so take best fit)
-	for (int32 FreeBlockIndex = 0; FreeBlockIndex < FreeBlocks.Num(); ++FreeBlockIndex)
+		int32 FindIndex = Algo::LowerBoundBy(FreeBlocks, AlignedSize, [](const FRHIPoolAllocationData* FreeBlock)
+			{
+				return FreeBlock->GetSize();
+			});
+		if (FindIndex < FreeBlocks.Num())
+		{
+			return FindIndex;
+		}
+	}
+	else
 	{
-		if (FreeBlocks[FreeBlockIndex]->GetSize() >= AlignedSize)
+		// Check all the free blocks (sorted by size, so take best fit)
+		for (int32 FreeBlockIndex = 0; FreeBlockIndex < FreeBlocks.Num(); ++FreeBlockIndex)
 		{
-			return FreeBlockIndex;
+			if (FreeBlocks[FreeBlockIndex]->GetSize() >= AlignedSize)
+			{
+				return FreeBlockIndex;
+			}
 		}
 	}
 
@@ -578,6 +612,7 @@ FRHIPoolAllocator::FRHIPoolAllocator(uint64 InDefaultPoolSize, uint32 InPoolAlig
 	MaxAllocationSize(InMaxAllocationSize),
 	FreeListOrder(InFreeListOrder),
 	bDefragEnabled(InDefragEnabled),
+	LastDefragPoolIndex(0),
 	TotalAllocatedBlocks(0)
 {}
 
@@ -669,49 +704,108 @@ void FRHIPoolAllocator::Defrag(uint32 InMaxCopySize, uint32& CurrentCopySize)
 		return;
 	}
 
+	TRACE_CPUPROFILER_EVENT_SCOPE(FRHIPoolAllocator::Defrag);
+
 	FScopeLock Lock(&CS);
 
-	TArray<FRHIMemoryPool*> SortedPools;
-	SortedPools.Reserve(Pools.Num());
+	TArray<FRHIMemoryPool*> SortedTargetPools;
+	SortedTargetPools.Reserve(Pools.Num());
+	uint32 MaxPoolsToDefrag = 0;
 	for (FRHIMemoryPool* Pool : Pools)
 	{
+		// collect all non full pools as target for defrag
 		if (Pool && !Pool->IsFull())
 		{
-			SortedPools.Add(Pool);
+			SortedTargetPools.Add(Pool);
+		}
+
+		// compute the usage percentage of the pool - don't defrag if used more than certain value
+		float Usage = Pool ? (float)Pool->GetUsedSize() / (float)Pool->GetPoolSize() : 1.0f;
+		if (Usage < GRHIPoolAllocatorDefragSizeFraction)
+		{
+			MaxPoolsToDefrag++;
 		}
 	}
 
 	// Need more than 1 pool at least
-	if (SortedPools.Num() < 2)
+	if (MaxPoolsToDefrag == 0 || SortedTargetPools.Num() < 2)
 	{
 		return;
 	}
 
-	Algo::Sort(SortedPools, [](const FRHIMemoryPool* InLHS, const FRHIMemoryPool* InRHS)
+	Algo::Sort(SortedTargetPools, [](const FRHIMemoryPool* InLHS, const FRHIMemoryPool* InRHS)
 		{
 			return InLHS->GetUsedSize() < InRHS->GetUsedSize();
 		});
 
-	// Build list all target pools (added reverse order because we want to move to the fullest pool first)
-	// Skip index 0 because that's the pool we want to clear to the target pools
 	TArray<FRHIMemoryPool*> TargetPools;
-	TargetPools.Reserve(SortedPools.Num());
-	for (int32 PoolIndex = SortedPools.Num() - 1; PoolIndex > 0; --PoolIndex)
+	if (GRHIPoolAllocatorDefragMaxPoolsToClear < 0)
 	{
-		TargetPools.Add(SortedPools[PoolIndex]);
-	}
+		// Build list all target pools (added reverse order because we want to move to the fullest pool first)
+		// Skip index 0 because that's the pool we want to clear to the target pools
+		TargetPools.Reserve(SortedTargetPools.Num());
+		for (int32 PoolIndex = SortedTargetPools.Num() - 1; PoolIndex > 0; --PoolIndex)
+		{
+			TargetPools.Add(SortedTargetPools[PoolIndex]);
+		}
 
-	for (int32 PoolIndex = 0; PoolIndex < SortedPools.Num() - 1; ++PoolIndex)
-	{
-		FRHIMemoryPool* PoolToClear = SortedPools[PoolIndex];
-		PoolToClear->TryClear(this, InMaxCopySize, CurrentCopySize, TargetPools);
+		for (int32 PoolIndex = 0; PoolIndex < SortedTargetPools.Num() - 1; ++PoolIndex)
+		{
+			FRHIMemoryPool* PoolToClear = SortedTargetPools[PoolIndex];
+			PoolToClear->TryClear(this, InMaxCopySize, CurrentCopySize, TargetPools);
 
-		// Remove last allocator since we will try and clear that one next		
-		TargetPools.Pop();
+			// Remove last allocator since we will try and clear that one next		
+			TargetPools.Pop();
 
-		if (CurrentCopySize >= InMaxCopySize)
+			if (CurrentCopySize >= InMaxCopySize)
+			{
+				break;
+			}
+		}
+	}
+	else
+	{
+		if (LastDefragPoolIndex >= Pools.Num())
 		{
-			break;
+			LastDefragPoolIndex = 0;
+		}
+		int32 DefraggedPoolCount = 0;
+		for (; LastDefragPoolIndex < Pools.Num(); LastDefragPoolIndex++)
+		{
+			FRHIMemoryPool* PoolToClear = Pools[LastDefragPoolIndex];
+
+			// skip if invalid, or too full
+			float Usage = PoolToClear ? (float)PoolToClear->GetUsedSize() / (float)PoolToClear->GetPoolSize() : 1.0f;
+			if (Usage >= GRHIPoolAllocatorDefragSizeFraction)
+			{
+				continue;
+			}
+
+			// Build list all target pools (added reverse order because we want to move to the fullest pool first)
+			TargetPools.Empty(SortedTargetPools.Num());
+			for (int32 SortedPoolIndex = SortedTargetPools.Num() - 1; SortedPoolIndex >= 0; --SortedPoolIndex)
+			{
+				FRHIMemoryPool* TargetPool = SortedTargetPools[SortedPoolIndex];
+				if (PoolToClear == TargetPool)
+				{
+					break;
+				}
+				TargetPools.Add(TargetPool);
+			}
+
+			// still have target pools?
+			if (TargetPools.Num() == 0)
+			{
+				continue;
+			}
+
+			PoolToClear->TryClear(this, InMaxCopySize, CurrentCopySize, TargetPools);
+			DefraggedPoolCount++;
+
+			if (CurrentCopySize >= InMaxCopySize || (DefraggedPoolCount >= GRHIPoolAllocatorDefragMaxPoolsToClear))
+			{
+				break;
+			}
 		}
 	}
 }
diff --git a/Engine/Source/Runtime/RHICore/Public/RHIPoolAllocator.h b/Engine/Source/Runtime/RHICore/Public/RHIPoolAllocator.h
--- a/Engine/Source/Runtime/RHICore/Public/RHIPoolAllocator.h
+++ b/Engine/Source/Runtime/RHICore/Public/RHIPoolAllocator.h
@@ -254,6 +254,9 @@ protected:
 	// Actual pools managing each a linked list of allocations
 	TArray<FRHIMemoryPool*> Pools;
 
+	// Last defrag pool index - defrag is time slices to reduce CPU overhead
+	int32 LastDefragPoolIndex;
+
 	// Allocation order of the pools when performing new allocations
 	// Can' sort the pools directly because pool index is stored in the allocation info
 	TArray<uint32> PoolAllocationOrder;
