Major refactor of transient allocator to support page table mapping implementation and extracted transient resources.

 - Implemented common transient page allocator in RHICore.
 - Implemented XBox specific GPU page table mapping allocator.
 - Extended RDG insights to support viewing heap visualization or page pool visualization.

diff --git a/Engine/Plugins/RenderGraphInsights/Source/RenderGraphInsights/Private/RenderGraphProvider.cpp b/Engine/Plugins/RenderGraphInsights/Source/RenderGraphInsights/Private/RenderGraphProvider.cpp
--- a/Engine/Plugins/RenderGraphInsights/Source/RenderGraphInsights/Private/RenderGraphProvider.cpp
+++ b/Engine/Plugins/RenderGraphInsights/Source/RenderGraphInsights/Private/RenderGraphProvider.cpp
@@ -66,10 +66,22 @@ FResourcePacket::FResourcePacket(const UE::Trace::IAnalyzer::FOnEventContext& Co
 	, bExtracted(Context.EventData.GetValue<bool>("IsExtracted"))
 	, bCulled(Context.EventData.GetValue<bool>("IsCulled"))
 	, bTransient(Context.EventData.GetValue<bool>("IsTransient"))
+	, bTransientUntracked(Context.EventData.GetValue<bool>("IsTransientUntracked"))
+	, bTransientCacheHit(Context.EventData.GetValue<bool>("IsTransientCacheHit"))
 {
-	TransientStats.HeapIndex = Context.EventData.GetValue<uint16>("TransientHeapIndex");
-	TransientStats.HeapOffsetMin = Context.EventData.GetValue<uint64>("TransientHeapOffsetMin");
-	TransientStats.HeapOffsetMax = Context.EventData.GetValue<uint64>("TransientHeapOffsetMax");
+	const auto TransientAllocationOffsetMins   = Context.EventData.GetArrayView<uint64>("TransientAllocationOffsetMins");
+	const auto TransientAllocationOffsetMaxs   = Context.EventData.GetArrayView<uint64>("TransientAllocationOffsetMaxs");
+	const auto TransientAllocationMemoryRanges = Context.EventData.GetArrayView<uint16>("TransientAllocationMemoryRanges");
+	check(TransientAllocationOffsetMins.Num() == TransientAllocationOffsetMaxs.Num() && TransientAllocationOffsetMaxs.Num() == TransientAllocationMemoryRanges.Num());
+
+	TransientAllocations.SetNum(TransientAllocationOffsetMins.Num());
+
+	for (int32 LocalIndex = 0; LocalIndex < TransientAllocationOffsetMins.Num(); ++LocalIndex)
+	{
+		TransientAllocations[LocalIndex].OffsetMin = TransientAllocationOffsetMins[LocalIndex];
+		TransientAllocations[LocalIndex].OffsetMax = TransientAllocationOffsetMaxs[LocalIndex];
+		TransientAllocations[LocalIndex].MemoryRangeIndex = TransientAllocationMemoryRanges[LocalIndex];
+	}
 
 	if (Passes.Num())
 	{
@@ -141,20 +153,21 @@ FGraphPacket::FGraphPacket(TraceServices::ILinearAllocator& Allocator, const UE:
 {
 	NormalizedPassDuration = (EndTime - StartTime) / double(PassCount);
 
-	TConstArrayView<uint64> WatermarkSizes(Context.EventData.GetArrayView<uint64>("TransientHeapWatermarkSizes"));
-	TConstArrayView<uint64> Capacities(Context.EventData.GetArrayView<uint64>("TransientHeapCapacities"));
-	check(WatermarkSizes.Num() == Capacities.Num());
+	const auto CommitSizes = Context.EventData.GetArrayView<uint64>("TransientMemoryCommitSizes");
+	const auto Capacities  = Context.EventData.GetArrayView<uint64>("TransientMemoryCapacities");
+	check(CommitSizes.Num() == Capacities.Num());
 
-	uint64 currentOffset = 0;
+	uint64 CurrentOffset = 0;
+	TransientMemoryRangeByteOffsets.Reserve(CommitSizes.Num());
 
-	for (int32 Index = 0; Index < WatermarkSizes.Num(); ++Index)
+	for (int32 Index = 0; Index < CommitSizes.Num(); ++Index)
 	{
-		auto& Heap = TransientHeapStats.Heaps.AddDefaulted_GetRef();
-		Heap.WatermarkSize = WatermarkSizes[Index];
-		Heap.Capacity = Capacities[Index];
+		auto& MemoryRange = TransientAllocationStats.MemoryRanges.AddDefaulted_GetRef();
+		MemoryRange.CommitSize = CommitSizes[Index];
+		MemoryRange.Capacity = Capacities[Index];
 
-		TransientHeapOffsets.Emplace(currentOffset);
-		currentOffset += Heap.WatermarkSize;
+		TransientMemoryRangeByteOffsets.Emplace(CurrentOffset);
+		CurrentOffset += MemoryRange.CommitSize;
 	}
 }
 
@@ -215,12 +228,12 @@ void FRenderGraphProvider::SetupResource(FResourcePacket& Resource)
 	const FPassPacket* FirstPass = CurrentGraph->GetPass(Resource.FirstPass);
 	const FPassPacket* LastPass  = CurrentGraph->GetPass(Resource.LastPass);
 
-	if (Resource.bExternal)
+	if (Resource.bExternal || Resource.bTransientUntracked)
 	{
 		FirstPass = CurrentGraph->GetProloguePass();
 	}
 
-	if (Resource.bExtracted)
+	if (Resource.bExtracted || Resource.bTransientUntracked)
 	{
 		LastPass = CurrentGraph->GetEpiloguePass();
 	}
diff --git a/Engine/Plugins/RenderGraphInsights/Source/RenderGraphInsights/Private/RenderGraphProvider.h b/Engine/Plugins/RenderGraphInsights/Source/RenderGraphInsights/Private/RenderGraphProvider.h
--- a/Engine/Plugins/RenderGraphInsights/Source/RenderGraphInsights/Private/RenderGraphProvider.h
+++ b/Engine/Plugins/RenderGraphInsights/Source/RenderGraphInsights/Private/RenderGraphProvider.h
@@ -45,7 +45,6 @@ public:
 	FString Name;
 	double StartTime{};
 	double EndTime{};
-	mutable uint32 VisibleIndex = kInvalidVisibleIndex;
 
 	FPacket(const UE::Trace::IAnalyzer::FOnEventContext& Context);
 };
@@ -66,6 +65,8 @@ class FScopePacket : public FPassIntervalPacket
 public:
 	uint32 Depth{};
 
+	mutable uint32 VisibleIndex = kInvalidVisibleIndex;
+
 	FScopePacket(const UE::Trace::IAnalyzer::FOnEventContext& Context);
 };
 
@@ -76,14 +77,18 @@ public:
 	uint32 Index{};
 	uint32 Order{};
 	uint64 SizeInBytes{};
-	FRHITransientResourceStats TransientStats;
+	FRHITransientAllocationStats::FAllocationArray TransientAllocations;
 
 	TArray<FRDGPassHandle> Passes;
 
+	mutable TArray<uint32, TInlineAllocator<1>> VisibleItems;
+
 	bool bExternal{};
 	bool bExtracted{};
 	bool bCulled{};
 	bool bTransient{};
+	bool bTransientUntracked{};
+	bool bTransientCacheHit{};
 
 	FResourcePacket(const UE::Trace::IAnalyzer::FOnEventContext& Context);
 };
@@ -118,6 +123,9 @@ class FPassPacket : public FPacket
 public:
 	TArray<FRDGTextureHandle> Textures;
 	TArray<FRDGBufferHandle> Buffers;
+
+	mutable uint32 VisibleIndex = kInvalidVisibleIndex;
+
 	FRDGPassHandle Handle;
 	FRDGPassHandle GraphicsForkPass;
 	FRDGPassHandle GraphicsJoinPass;
@@ -150,11 +158,13 @@ public:
 	TMap<FRDGBufferHandle, const FBufferPacket*> BufferHandleToPreviousOwner;
 	TMap<FRDGTextureHandle, const FTexturePacket*> TextureHandleToPreviousOwner;
 
+	mutable uint32 VisibleIndex = kInvalidVisibleIndex;
+
 	uint32 ScopeDepth{};
 	uint32 PassCount{};
 
-	FRHITransientHeapStats TransientHeapStats;
-	TArray<uint64, TInlineAllocator<8>> TransientHeapOffsets;
+	FRHITransientAllocationStats TransientAllocationStats;
+	TArray<uint64> TransientMemoryRangeByteOffsets;
 
 	FGraphPacket(TraceServices::ILinearAllocator& Allocator, const UE::Trace::IAnalyzer::FOnEventContext& Context);
 
diff --git a/Engine/Plugins/RenderGraphInsights/Source/RenderGraphInsights/Private/RenderGraphTrack.cpp b/Engine/Plugins/RenderGraphInsights/Source/RenderGraphInsights/Private/RenderGraphTrack.cpp
--- a/Engine/Plugins/RenderGraphInsights/Source/RenderGraphInsights/Private/RenderGraphTrack.cpp
+++ b/Engine/Plugins/RenderGraphInsights/Source/RenderGraphInsights/Private/RenderGraphTrack.cpp
@@ -38,6 +38,7 @@ constexpr uint32 kNoParameterPassColor = 0xff4D4D4D;
 constexpr uint32 kAsyncComputePassColor = 0xff2D7f2d;
 constexpr uint32 kTextureColor = 0xff89cff0;
 constexpr uint32 kBufferColor = 0xff66D066;
+constexpr uint32 kUntrackedColor = 0x2f2f2f2f;
 
 static uint32 GetPassColor(const FPassPacket& Packet)
 {
@@ -81,6 +82,13 @@ inline uint32 GetResourceColorBySize(uint64 Size, uint64 MaxSize)
 	return FLinearColor::LerpUsingHSV(Low, High, Percentage).ToFColor(false).ToPackedARGB();
 }
 
+inline uint32 GetResourceColorByTransientCache(bool bHit)
+{
+	static const uint32 HitColor = FLinearColor(0.01, 0.01, 0.01, 0.25f).ToFColor(false).ToPackedARGB();
+	static const uint32 MissColor = FLinearColor(1.0, 0.1, 0.1, 1.0f).ToFColor(false).ToPackedARGB();
+	return bHit ? HitColor : MissColor;
+}
+
 ///////////////////////////////////////////////////////////////////////////////
 
 template <typename VisibleItemType>
@@ -89,8 +97,8 @@ static VisibleItemType& AddVisibleItem(TArray<VisibleItemType>& VisibleItems, co
 	const uint32 VisibleIndex = VisibleItems.Num();
 
 	VisibleItemType& VisibleItem = VisibleItems.Emplace_GetRef(InVisibleItem);
-	check(VisibleItem.Packet->VisibleIndex == kInvalidVisibleIndex);
-	VisibleItem.Packet->VisibleIndex = VisibleIndex;
+	check(VisibleItem.GetPacket().VisibleIndex == kInvalidVisibleIndex);
+	VisibleItem.GetPacket().VisibleIndex = VisibleIndex;
 	VisibleItem.Index = VisibleIndex;
 	return VisibleItem;
 }
@@ -100,12 +108,34 @@ static void ResetVisibleItemArray(TArray<VisibleItemType>& VisibleItems)
 {
 	for (VisibleItemType& VisibleItem : VisibleItems)
 	{
-		check(VisibleItem.Packet->VisibleIndex == VisibleItem.Index);
-		VisibleItem.Packet->VisibleIndex = kInvalidVisibleIndex;
+		check(VisibleItem.GetPacket().VisibleIndex == VisibleItem.Index);
+		VisibleItem.GetPacket().VisibleIndex = kInvalidVisibleIndex;
 	}
 	VisibleItems.Empty();
 }
 
+template <typename VisibleResourceType>
+static VisibleResourceType& AddVisibleResource(TArray<VisibleResourceType>& VisibleResources, const VisibleResourceType& InVisibleResource)
+{
+	const uint32 VisibleIndex = VisibleResources.Num();
+
+	VisibleResourceType& VisibleResource = VisibleResources.Emplace_GetRef(InVisibleResource);
+	check(!VisibleResource.GetPacket().VisibleItems.Contains(VisibleIndex));
+	VisibleResource.GetPacket().VisibleItems.Emplace(VisibleIndex);
+	VisibleResource.Index = VisibleIndex;
+	return VisibleResource;
+}
+
+template <typename VisibleResourceType>
+static void ResetVisibleResourceArray(TArray<VisibleResourceType>& VisibleResources)
+{
+	for (VisibleResourceType& VisibleResource : VisibleResources)
+	{
+		VisibleResource.GetPacket().VisibleItems.Reset();
+	}
+	VisibleResources.Empty();
+}
+
 ///////////////////////////////////////////////////////////////////////////////
 
 class FRenderGraphTimingEvent : public FTimingEvent
@@ -127,6 +157,7 @@ public:
 
 	virtual const FVisibleItem& GetItem() const = 0;
 	virtual const FPacket& GetPacket() const = 0;
+	virtual uint32 GetAllocationIndex() const { return 0; };
 };
 
 template <typename ItemType>
@@ -181,6 +212,8 @@ public:
 	FVisibleTextureEvent(const TSharedRef<const FBaseTimingTrack>& InTrack, const FVisibleTexture& InItem)
 		: TRenderGraphTimingEventHelper<FVisibleTexture>(InTrack, InItem)
 	{}
+
+	uint32 GetAllocationIndex() const override { return GetItem().AllocationIndex; };
 };
 
 class FVisibleBufferEvent final : public TRenderGraphTimingEventHelper<FVisibleBuffer>
@@ -190,6 +223,8 @@ public:
 	FVisibleBufferEvent(const TSharedRef<const FBaseTimingTrack>& InTrack, const FVisibleBuffer& InItem)
 		: TRenderGraphTimingEventHelper<FVisibleBuffer>(InTrack, InItem)
 	{}
+
+	uint32 GetAllocationIndex() const override { return GetItem().AllocationIndex; };
 };
 
 inline const FPassPacket* TryGetPassPacket(const TSharedPtr<const ITimingEvent>& TimingEventBase)
@@ -216,6 +251,7 @@ public:
 		: Event(InEvent)
 		, Packet(&InEvent->GetPacket())
 		, Graph(Packet->Graph)
+		, AllocationIndex(InEvent->GetAllocationIndex())
 	{}
 	~FPacketFilter() {}
 
@@ -265,6 +301,11 @@ public:
 		return *Graph;
 	}
 
+	uint32 GetAllocationIndex() const
+	{
+		return AllocationIndex;
+	}
+
 private:
 	//! ITimingEventFilter
 	bool FilterTrack(const FBaseTimingTrack& InTrack) const override { return true; }
@@ -276,6 +317,7 @@ private:
 	TSharedPtr<const FRenderGraphTimingEvent> Event;
 	const FPacket* Packet{};
 	const FGraphPacket* Graph{};
+	uint32 AllocationIndex = 0;
 };
 
 inline bool TryFilterPacket(const FPacketFilter* Filter, const FPacket& Packet)
@@ -322,8 +364,8 @@ void FVisibleGraph::Reset()
 {
 	ResetVisibleItemArray(Scopes);
 	ResetVisibleItemArray(Passes);
-	ResetVisibleItemArray(Textures);
-	ResetVisibleItemArray(Buffers);
+	ResetVisibleResourceArray(Textures);
+	ResetVisibleResourceArray(Buffers);
 }
 
 void FVisibleGraph::AddScope(const FVisibleScope& VisibleScope)
@@ -343,14 +385,14 @@ void FVisibleGraph::AddPass(const FVisiblePass& InVisiblePass)
 
 void FVisibleGraph::AddTexture(const FVisibleTexture& VisibleTexture)
 {
-	AddVisibleItem(Textures, VisibleTexture);
+	AddVisibleResource(Textures, VisibleTexture);
 
 	DepthH = FMath::Max(DepthH, VisibleTexture.DepthY + VisibleTexture.DepthH);
 }
 
 void FVisibleGraph::AddBuffer(const FVisibleBuffer& VisibleBuffer)
 {
-	AddVisibleItem(Buffers, VisibleBuffer);
+	AddVisibleResource(Buffers, VisibleBuffer);
 
 	DepthH = FMath::Max(DepthH, VisibleBuffer.DepthY + VisibleBuffer.DepthH);
 }
@@ -421,32 +463,61 @@ FRenderGraphTrack::FRenderGraphTrack(const FRenderGraphTimingViewSession& InShar
 
 uint32 FRenderGraphTrack::GetTextureColor(const FTexturePacket& Texture, uint64 MaxSizeInBytes) const
 {
+	if (Texture.bTransientUntracked)
+	{
+		return kUntrackedColor;
+	}
+
 	if (ResourceColor == EResourceColor::Type)
 	{
 		return kTextureColor;
 	}
-
-	return GetResourceColorBySize(Texture.SizeInBytes, MaxSizeInBytes);
+	else if (ResourceColor == EResourceColor::TransientCache)
+	{
+		return GetResourceColorByTransientCache(!Texture.bTransient || Texture.bTransientCacheHit);
+	}
+	else
+	{
+		return GetResourceColorBySize(Texture.SizeInBytes, MaxSizeInBytes);
+	}
 }
 
-uint32 FRenderGraphTrack::GetBufferColor(const FBufferPacket& Texture, uint64 MaxSizeInBytes) const
+uint32 FRenderGraphTrack::GetBufferColor(const FBufferPacket& Buffer, uint64 MaxSizeInBytes) const
 {
+	if (Buffer.bTransientUntracked)
+	{
+		return kUntrackedColor;
+	}
+
 	if (ResourceColor == EResourceColor::Type)
 	{
 		return kBufferColor;
 	}
-
-	return GetResourceColorBySize(Texture.SizeInBytes, MaxSizeInBytes);
+	else if (ResourceColor == EResourceColor::TransientCache)
+	{
+		return GetResourceColorByTransientCache(!Buffer.bTransient || Buffer.bTransientCacheHit);
+	}
+	else
+	{
+		return GetResourceColorBySize(Buffer.SizeInBytes, MaxSizeInBytes);
+	}
 }
 
-float FRenderGraphTrack::TransientHeapOffsetToDepth(uint64 HeapOffset) const
+float FRenderGraphTrack::TransientByteOffsetToDepth(uint64 MemoryOffset) const
 {
-	return (float)HeapOffset / TransientHeapBytesPerDepthSlot;
+	if (MemoryOffset == 0)
+	{
+		return 0;
+	}
+	else
+	{
+		return FMath::Max((float)MemoryOffset / TransientBytesPerDepthSlot, 1.0f / 16.0f);
+	}
 }
 
-float FRenderGraphTrack::TransientHeapOffsetToSlateY(const FTimingViewLayout& Layout, uint64 HeapOffset) const
+float FRenderGraphTrack::TransientByteOffsetToSlateY(const FTimingViewLayout& Layout, uint64 MemoryOffset) const
 {
-	return (Layout.EventH + Layout.EventDY) * TransientHeapOffsetToDepth(HeapOffset);
+	return (Layout.EventH + Layout.EventDY) * TransientByteOffsetToDepth(MemoryOffset);
 }
 
 void FRenderGraphTrack::Reset()
@@ -826,10 +897,10 @@ void FRenderGraphTrack::Draw(const ITimingTrackDrawContext& Context) const
 		}
 
 		////////////////////////////////////////////////////////////////////////////////////
-		// Draw transient heap vertical lines and size markers
+		// Draw transient memory vertical lines and size markers
 		////////////////////////////////////////////////////////////////////////////////////
 
-		if (Visualizer == EVisualizer::TransientHeaps)
+		if (Visualizer == EVisualizer::TransientMemory)
 		{
 			const FLinearColor MBLineColor(0.8f, 0.1f, 0.1f, 0.2f);
 			const FLinearColor MBTextColor(0.8f, 0.1f, 0.1f, 0.8f);
@@ -840,7 +911,7 @@ void FRenderGraphTrack::Draw(const ITimingTrackDrawContext& Context) const
 
 			const uint64 MiB = 1024 * 1024;
 
-			const auto DrawMibText = [&](float Depth, uint64 Offset, uint32 HeapIndex, uint32 HeapSizeInMB)
+			const auto DrawMibText = [&](float Depth, uint64 Offset, uint32 MemoryRangeIndex, uint32 MemoryRangeSizeInMB)
 			{
 				FString Text;
 				FLinearColor TextColor = MBTextColor;
@@ -851,7 +922,7 @@ void FRenderGraphTrack::Draw(const ITimingTrackDrawContext& Context) const
 				}
 				else
 				{
-					Text = FString::Printf(TEXT("Heap[%d] %uMiB"), HeapIndex, HeapSizeInMB);
+					Text = FString::Printf(TEXT("Memory Range [%d] %uMiB"), MemoryRangeIndex, MemoryRangeSizeInMB);
 					TextColor.A = 1.0f;
 				}
 
@@ -868,20 +939,20 @@ void FRenderGraphTrack::Draw(const ITimingTrackDrawContext& Context) const
 
 			const float SinglePixelDepth = 1.0f / Layout.EventH;
 
-			for (int32 HeapIndex = 0; HeapIndex < Graph.TransientHeapOffsets.Num(); ++HeapIndex)
+			for (int32 MemoryRangeIndex = 0; MemoryRangeIndex < Graph.TransientMemoryRangeByteOffsets.Num(); ++MemoryRangeIndex)
 			{
-				const float HeapStartDepth = StartDepth + VisibleGraph.TransientHeapDepths[HeapIndex];
+				const float MemoryRangeStartDepth = StartDepth + VisibleGraph.TransientMemoryRangeDepthOffsets[MemoryRangeIndex];
 
-				const auto& TransientHeap = Graph.TransientHeapStats.Heaps[HeapIndex];
+				const auto& TransientMemoryRange = Graph.TransientAllocationStats.MemoryRanges[MemoryRangeIndex];
 
 				uint64 Offset = 0;
 				uint32 OffsetIndex = 0;
 
 				const uint32 TicksPerText = 4;
 
-				while (Offset < TransientHeap.WatermarkSize)
+				while (Offset < TransientMemoryRange.CommitSize)
 				{
-					const float Depth = HeapStartDepth + TransientHeapOffsetToDepth(Offset);
+					const float Depth = MemoryRangeStartDepth + TransientByteOffsetToDepth(Offset);
 
 					FLinearColor LineColor = MBLineColor;
 
@@ -897,13 +968,13 @@ void FRenderGraphTrack::Draw(const ITimingTrackDrawContext& Context) const
 
 					if (bDrawText && bMajorTick)
 					{
-						LineStartX = DrawMibText(Depth, Offset, HeapIndex, TransientHeap.Capacity / MiB);
+						LineStartX = DrawMibText(Depth, Offset, MemoryRangeIndex, TransientMemoryRange.Capacity / MiB);
 						LineWidth -= LineStartX - VisibleGraph.SlateX;
 					}
 
 					DrawHelper.DrawBox(*this, LineStartX, Depth, LineWidth, SinglePixelDepth, LineColor, EDrawLayer::Background);
 
-					Offset += TransientHeapBytesPerDepthSlot * (Layout.bIsCompactMode ? 3 : 1);
+					Offset += TransientBytesPerDepthSlot * (Layout.bIsCompactMode ? 3 : 1);
 					OffsetIndex++;
 				}
 			}
@@ -987,18 +1058,8 @@ void FRenderGraphTrack::BuildDrawState(FRenderGraphTrackDrawStateBuilder& Builde
 		VisibleGraph.DepthH = DepthOffset;
 	};
 
-	if (Visualizer == EVisualizer::TransientHeaps)
+	if (Visualizer == EVisualizer::TransientMemory)
 	{
-		struct FResourceEntry
-		{
-			double StartTime{};
-			double EndTime{};
-			FRHITransientResourceStats TransientStats;
-			ERDGParentResourceType Type = ERDGParentResourceType::Texture;
-		};
-
-		TArray<FResourceEntry> Resources;
-
 		TraceServices::FAnalysisSessionReadScope SessionReadScope(SharedData.GetAnalysisSession());
 
 		const FRenderGraphProvider::TGraphTimeline& GraphTimeline = RenderGraphProvider->GetGraphTimeline();
@@ -1019,48 +1080,105 @@ void FRenderGraphTrack::BuildDrawState(FRenderGraphTrackDrawStateBuilder& Builde
 
 			const float YOffset = Layout.GetLaneY(DepthOffset);
 
+			const auto IsPacketCulled = [&](const FResourcePacket& Packet)
+			{
+				if (!Packet.bTransient)
+				{
+					return true;
+				}
+
+				if (!Packet.bTransientUntracked && !ShowTracked())
+				{
+					return true;
+				}
+
+				if (FilterSize > 0.0f && Packet.SizeInBytes < uint64(FilterSize * 1024.0f * 1024.0f))
+				{
+					return true;
+				}
+
+				if (!FilterText.IsEmpty() && !Packet.Name.Contains(FilterText))
+				{
+					return true;
+				}
+
+				return false;
+			};
+
+			uint64 MaxSizeInBytes = 0;
+
 			for (uint32 TextureIndex = 0, TextureCount = Graph->Textures.Num(); TextureIndex < TextureCount; ++TextureIndex)
 			{
 				FTexturePacket& Texture = Graph->Textures[TextureIndex];
 
-				if (!Texture.bTransient)
+				if (!IsPacketCulled(Texture))
+				{
+					MaxSizeInBytes = FMath::Max(MaxSizeInBytes, Texture.SizeInBytes);
+				}
+			}
+
+			for (uint32 BufferIndex = 0, BufferCount = Graph->Buffers.Num(); BufferIndex < BufferCount; ++BufferIndex)
+			{
+				FBufferPacket& Buffer = Graph->Buffers[BufferIndex];
+
+				if (!IsPacketCulled(Buffer))
+				{
+					MaxSizeInBytes = FMath::Max(MaxSizeInBytes, Buffer.SizeInBytes);
+				}
+			}
+
+			for (uint32 TextureIndex = 0, TextureCount = Graph->Textures.Num(); TextureIndex < TextureCount; ++TextureIndex)
+			{
+				FTexturePacket& Texture = Graph->Textures[TextureIndex];
+
+				if (IsPacketCulled(Texture))
 				{
 					continue;
 				}
 
-				const float HeapOffsetDepthBase = VisibleGraph.TransientHeapDepths[Texture.TransientStats.HeapIndex];
+				for (int32 AllocationIndex = 0; AllocationIndex < Texture.TransientAllocations.Num(); ++AllocationIndex)
+				{
+					const FRHITransientAllocationStats::FAllocation& Allocation = Texture.TransientAllocations[AllocationIndex];
 
-				const float DepthY = TransientHeapOffsetToDepth(Texture.TransientStats.HeapOffsetMin) + HeapOffsetDepthBase + DepthOffset;
-				const float DepthH = TransientHeapOffsetToDepth(Texture.TransientStats.HeapOffsetMax - Texture.TransientStats.HeapOffsetMin);
-				const double StartTime = Texture.StartTime + SinglePixelTimeMargin;
-				const double EndTime = Texture.EndTime;
-				const uint32 Color = kTextureColor;
+					const float MemoryRangeOffsetDepthBase = VisibleGraph.TransientMemoryRangeDepthOffsets[Allocation.MemoryRangeIndex];
 
-				const FVisibleTexture VisibleTexture(Viewport, Texture, Color, StartTime, EndTime, DepthY, DepthH);
-				VisibleTexture.Draw(Builder);
-				VisibleGraph.AddTexture(VisibleTexture);
+					const float DepthY = TransientByteOffsetToDepth(Allocation.OffsetMin) + MemoryRangeOffsetDepthBase + DepthOffset;
+					const float DepthH = TransientByteOffsetToDepth(Allocation.OffsetMax - Allocation.OffsetMin);
+					const double StartTime = Texture.StartTime + SinglePixelTimeMargin;
+					const double EndTime = Texture.EndTime;
+					const uint32 Color = GetTextureColor(Texture, MaxSizeInBytes);
+
+					const FVisibleTexture VisibleTexture(Viewport, Texture, Color, StartTime, EndTime, DepthY, DepthH, AllocationIndex);
+					VisibleTexture.Draw(Builder);
+					VisibleGraph.AddTexture(VisibleTexture);
+				}
 			}
 
 			for (uint32 BufferIndex = 0, BufferCount = Graph->Buffers.Num(); BufferIndex < BufferCount; ++BufferIndex)
 			{
 				FBufferPacket& Buffer = Graph->Buffers[BufferIndex];
 
-				if (!Buffer.bTransient)
+				if (IsPacketCulled(Buffer))
 				{
 					continue;
 				}
-				
-				const float HeapOffsetDepthBase = VisibleGraph.TransientHeapDepths[Buffer.TransientStats.HeapIndex];
-
-				const float DepthY = TransientHeapOffsetToDepth(Buffer.TransientStats.HeapOffsetMin) + HeapOffsetDepthBase + DepthOffset;
-				const float DepthH = TransientHeapOffsetToDepth(Buffer.TransientStats.HeapOffsetMax - Buffer.TransientStats.HeapOffsetMin);
-				const double StartTime = Buffer.StartTime + SinglePixelTimeMargin;
-				const double EndTime = Buffer.EndTime;
-				const uint32 Color = kBufferColor;
-
-				const FVisibleBuffer VisibleBuffer(Viewport, Buffer, Color, StartTime, EndTime, DepthY, DepthH);
-				VisibleBuffer.Draw(Builder);
-				VisibleGraph.AddBuffer(VisibleBuffer);
+
+				for (int32 AllocationIndex = 0; AllocationIndex < Buffer.TransientAllocations.Num(); ++AllocationIndex)
+				{
+					const FRHITransientAllocationStats::FAllocation& Allocation = Buffer.TransientAllocations[AllocationIndex];
+
+					const float MemoryRangeOffsetDepthBase = VisibleGraph.TransientMemoryRangeDepthOffsets[Allocation.MemoryRangeIndex];
+
+					const float DepthY = TransientByteOffsetToDepth(Allocation.OffsetMin) + MemoryRangeOffsetDepthBase + DepthOffset;
+					const float DepthH = TransientByteOffsetToDepth(Allocation.OffsetMax - Allocation.OffsetMin);
+					const double StartTime = Buffer.StartTime + SinglePixelTimeMargin;
+					const double EndTime = Buffer.EndTime;
+					const uint32 Color = GetBufferColor(Buffer, MaxSizeInBytes);
+
+					const FVisibleBuffer VisibleBuffer(Viewport, Buffer, Color, StartTime, EndTime, DepthY, DepthH, AllocationIndex);
+					VisibleBuffer.Draw(Builder);
+					VisibleGraph.AddBuffer(VisibleBuffer);
+				}
 			}
 
 			return TraceServices::EEventEnumerate::Continue;
@@ -1110,7 +1228,7 @@ void FRenderGraphTrack::BuildDrawState(FRenderGraphTrackDrawStateBuilder& Builde
 
 			const auto IsPacketCulled = [&](const FResourcePacket& Packet)
 			{
-				if (Packet.bCulled)
+				if (Packet.bCulled || Packet.bTransientUntracked)
 				{
 					return true;
 				}
@@ -1125,7 +1243,17 @@ void FRenderGraphTrack::BuildDrawState(FRenderGraphTrackDrawStateBuilder& Builde
 					return true;
 				}
 
-				if (!Packet.bTransient && !Packet.bExternal && !ShowPooled())
+				if (Packet.bExtracted && !ShowExtracted())
+				{
+					return true;
+				}
+
+				if (!Packet.bExternal && !Packet.bExtracted && !ShowInternal())
+				{
+					return true;
+				}
+
+				if (!Packet.bTransient && !ShowPooled())
 				{
 					return true;
 				}
@@ -1348,10 +1476,11 @@ void FRenderGraphTrack::BuildFilteredDrawState(FRenderGraphTrackDrawStateBuilder
 			for (FRDGTextureHandle TextureHandle : Pass.Textures)
 			{
 				const FTexturePacket& Texture = *Graph.GetTexture(TextureHandle);
-				if (const FVisibleTexture* VisibleTexture = VisibleGraph->GetVisibleTexture(Texture))
+
+				VisibleGraph->EnumerateVisibleTextures(Texture, [&](const FVisibleTexture& VisibleTexture)
 				{
-					VisibleItems.Add(VisibleTexture);
-				}
+					VisibleItems.Add(&VisibleTexture);
+				});
 			}
 		}
 
@@ -1360,10 +1489,11 @@ void FRenderGraphTrack::BuildFilteredDrawState(FRenderGraphTrackDrawStateBuilder
 			for (FRDGBufferHandle BufferHandle : Pass.Buffers)
 			{
 				const FBufferPacket& Buffer = *Graph.GetBuffer(BufferHandle);
-				if (const FVisibleBuffer* VisibleBuffer = VisibleGraph->GetVisibleBuffer(Buffer))
+
+				VisibleGraph->EnumerateVisibleBuffers(Buffer, [&](const FVisibleBuffer& VisibleBuffer)
 				{
-					VisibleItems.Add(VisibleBuffer);
-				}
+					VisibleItems.Add(&VisibleBuffer);
+				});
 			}
 		}
 
@@ -1396,35 +1526,38 @@ void FRenderGraphTrack::BuildFilteredDrawState(FRenderGraphTrackDrawStateBuilder
 
 	const FTimingViewLayout& Layout = Context.GetViewport().GetLayout();
 
-	const auto AddResourcePassEvents = [&](const FVisibleItem& VisibleResource, TArrayView<const FRDGPassHandle> Passes)
+	const auto AddResourcePassEvents = [&](const auto& VisibleResource, TArrayView<const FRDGPassHandle> Passes, uint32 AllocationIndex)
 	{
 		VisibleResource.Draw(Builder);
 
-		for (FRDGPassHandle PassHandle : Passes)
+		if (VisibleResource.AllocationIndex == AllocationIndex)
 		{
-			const FPassPacket& Pass = *Graph.GetPass(PassHandle);
-			const FVisiblePass& VisiblePass = VisibleGraph->GetVisiblePass(Pass);
-
-			VisiblePass.Draw(Builder);
+			for (FRDGPassHandle PassHandle : Passes)
+			{
+				const FPassPacket& Pass = *Graph.GetPass(PassHandle);
+				const FVisiblePass& VisiblePass = VisibleGraph->GetVisiblePass(Pass);
 
-			const float DepthY = VisiblePass.DepthY + VisiblePass.DepthH;
+				VisiblePass.Draw(Builder);
 
-			FSplinePrimitive Spline;
-			Spline.Start.X = VisiblePass.SlateX + (VisiblePass.SlateW * 0.5f);
-			Spline.Start.Y = DepthY;
-			Spline.StartDir = FVector2D(0, -1);
-			Spline.End = FVector2D(0, VisibleResource.DepthY - DepthY);
-			Spline.EndDir = FVector2D(0, 1);
-			Spline.Thickness = 1.0f;
-			Spline.Tint = FLinearColor(0.8f, 0.8f, 0.8f, 0.7f);
-			Builder.AddSpline(Spline);
+				const float DepthY = VisiblePass.DepthY + VisiblePass.DepthH;
 
-			for (const FVisibleScope& VisibleScope : VisibleGraph->Scopes)
-			{
-				const FScopePacket& Scope = VisibleScope.GetPacket();
-				if (Intersects(Scope, Pass))
+				FSplinePrimitive Spline;
+				Spline.Start.X = VisiblePass.SlateX + (VisiblePass.SlateW * 0.5f);
+				Spline.Start.Y = DepthY;
+				Spline.StartDir = FVector2D(0, -1);
+				Spline.End = FVector2D(0, VisibleResource.DepthY - DepthY);
+				Spline.EndDir = FVector2D(0, 1);
+				Spline.Thickness = 1.0f;
+				Spline.Tint = FLinearColor(0.8f, 0.8f, 0.8f, 0.7f);
+				Builder.AddSpline(Spline);
+
+				for (const FVisibleScope& VisibleScope : VisibleGraph->Scopes)
 				{
-					VisibleItems.Add(&VisibleScope);
+					const FScopePacket& Scope = VisibleScope.GetPacket();
+					if (Intersects(Scope, Pass))
+					{
+						VisibleItems.Add(&VisibleScope);
+					}
 				}
 			}
 		}
@@ -1438,7 +1571,7 @@ void FRenderGraphTrack::BuildFilteredDrawState(FRenderGraphTrackDrawStateBuilder
 
 			if (PacketFilter->FilterPacketExact(Texture))
 			{
-				AddResourcePassEvents(VisibleTexture, Texture.Passes);
+				AddResourcePassEvents(VisibleTexture, Texture.Passes, PacketFilter->GetAllocationIndex());
 			}
 		}
 	}
@@ -1451,7 +1584,7 @@ void FRenderGraphTrack::BuildFilteredDrawState(FRenderGraphTrackDrawStateBuilder
 
 			if (PacketFilter->FilterPacketExact(Buffer))
 			{
-				AddResourcePassEvents(VisibleBuffer, Buffer.Passes);
+				AddResourcePassEvents(VisibleBuffer, Buffer.Passes, PacketFilter->GetAllocationIndex());
 			}
 		}
 	}
@@ -1560,20 +1693,78 @@ void FRenderGraphTrack::BuildContextMenu(FMenuBuilder& MenuBuilder)
 			FUIAction(
 				FExecuteAction::CreateLambda([this, TimingViewSession]()
 				{
-					Visualizer = EVisualizer::TransientHeaps;
+					Visualizer = EVisualizer::TransientMemory;
 					TimingViewSession->ResetSelectedEvent();
 					TimingViewSession->ResetEventFilter();
 					ResourceShow = EResourceShow::All;
 					SetDirtyFlag();
 				}),
 				FCanExecuteAction(),
-				FIsActionChecked::CreateLambda([this]() { return Visualizer == EVisualizer::TransientHeaps; })
+				FIsActionChecked::CreateLambda([this]() { return Visualizer == EVisualizer::TransientMemory; })
 			),
 			NAME_None,
 			EUserInterfaceActionType::RadioButton
 		);
 	}
 	MenuBuilder.EndSection();
+	
+	MenuBuilder.BeginSection("Color", LOCTEXT("ColorMenuHeader", "Track Resource Coloration"));
+
+	MenuBuilder.AddMenuEntry
+	(
+		LOCTEXT("ColorType", "By Type"),
+		LOCTEXT("ColorType_Tooltip", "Each type of resource has a unique color."),
+		FSlateIcon(),
+		FUIAction(
+			FExecuteAction::CreateLambda([this]()
+			{
+				ResourceColor = EResourceColor::Type;
+				SetDirtyFlag();
+			}),
+			FCanExecuteAction(),
+			FIsActionChecked::CreateLambda([this]() { return ResourceColor == EResourceColor::Type; })
+		),
+		NAME_None,
+		EUserInterfaceActionType::RadioButton
+	);
+
+	MenuBuilder.AddMenuEntry
+	(
+		LOCTEXT("ColorSize", "By Size"),
+		LOCTEXT("ColorSize_Tooltip", "Larger resources are more brightly colored."),
+		FSlateIcon(),
+		FUIAction(
+			FExecuteAction::CreateLambda([this]()
+			{
+				ResourceColor = EResourceColor::Size;
+				SetDirtyFlag();
+			}),
+			FCanExecuteAction(),
+			FIsActionChecked::CreateLambda([this]() { return ResourceColor == EResourceColor::Size; })
+		),
+		NAME_None,
+		EUserInterfaceActionType::RadioButton
+	);
+		
+	MenuBuilder.AddMenuEntry
+	(
+		LOCTEXT("ColorTransientCache", "By Transient Cache"),
+		LOCTEXT("ColorTransientCache_Tooltip", "Resources that cause a cache miss are brightly colored."),
+		FSlateIcon(),
+		FUIAction(
+			FExecuteAction::CreateLambda([this]()
+			{
+				ResourceColor = EResourceColor::TransientCache;
+				SetDirtyFlag();
+			}),
+			FCanExecuteAction(),
+			FIsActionChecked::CreateLambda([this]() { return ResourceColor == EResourceColor::TransientCache; })
+		),
+		NAME_None,
+		EUserInterfaceActionType::RadioButton
+	);
+
+	MenuBuilder.EndSection();
 
 	if (Visualizer == EVisualizer::Resources)
 	{
@@ -1632,7 +1823,7 @@ void FRenderGraphTrack::BuildContextMenu(FMenuBuilder& MenuBuilder)
 			NAME_None,
 			EUserInterfaceActionType::ToggleButton
 		);
-		
+
 		MenuBuilder.AddMenuEntry
 		(
 			LOCTEXT("ShowExternal", "Show External"),
@@ -1651,6 +1842,42 @@ void FRenderGraphTrack::BuildContextMenu(FMenuBuilder& MenuBuilder)
 			EUserInterfaceActionType::ToggleButton
 		);
 		
+		MenuBuilder.AddMenuEntry
+		(
+			LOCTEXT("ShowInternal", "Show Internal"),
+			LOCTEXT("ShowInternal_Tooltip", "Show internal resources in the lifetime view."),
+			FSlateIcon(),
+			FUIAction(
+				FExecuteAction::CreateLambda([this]()
+				{
+					ResourceShow ^= EResourceShow::Internal;
+					SetDirtyFlag();
+				}),
+				FCanExecuteAction(),
+				FIsActionChecked::CreateLambda([this]() { return EnumHasAnyFlags(ResourceShow, EResourceShow::Internal); })
+			),
+			NAME_None,
+			EUserInterfaceActionType::ToggleButton
+		);
+		
+		MenuBuilder.AddMenuEntry
+		(
+			LOCTEXT("ShowExtracted", "Show Extracted"),
+			LOCTEXT("ShowExtracted_Tooltip", "Show extracted resources in the lifetime view."),
+			FSlateIcon(),
+			FUIAction(
+				FExecuteAction::CreateLambda([this]()
+				{
+					ResourceShow ^= EResourceShow::Extracted;
+					SetDirtyFlag();
+				}),
+				FCanExecuteAction(),
+				FIsActionChecked::CreateLambda([this]() { return EnumHasAnyFlags(ResourceShow, EResourceShow::Extracted); })
+			),
+			NAME_None,
+			EUserInterfaceActionType::ToggleButton
+		);
+
 		MenuBuilder.AddMenuEntry
 		(
 			LOCTEXT("ShowPooled", "Show Pooled"),
@@ -1668,7 +1895,7 @@ void FRenderGraphTrack::BuildContextMenu(FMenuBuilder& MenuBuilder)
 			NAME_None,
 			EUserInterfaceActionType::ToggleButton
 		);
-
+		
 		MenuBuilder.EndSection();
 		MenuBuilder.BeginSection("Sort", LOCTEXT("SortMenuHeader", "Track Sort By"));
 
@@ -1762,91 +1989,9 @@ void FRenderGraphTrack::BuildContextMenu(FMenuBuilder& MenuBuilder)
 			EUserInterfaceActionType::RadioButton
 		);
 
-		MenuBuilder.EndSection();
-		MenuBuilder.BeginSection("Color", LOCTEXT("ColorMenuHeader", "Track Resource Coloration"));
-
-		MenuBuilder.AddMenuEntry
-		(
-			LOCTEXT("ColorType", "By Type"),
-			LOCTEXT("ColorType_Tooltip", "Each type of resource has a unique color."),
-			FSlateIcon(),
-			FUIAction(
-				FExecuteAction::CreateLambda([this]()
-				{
-					ResourceColor = EResourceColor::Type;
-					SetDirtyFlag();
-				}),
-				FCanExecuteAction(),
-				FIsActionChecked::CreateLambda([this]() { return ResourceColor == EResourceColor::Type; })
-			),
-			NAME_None,
-			EUserInterfaceActionType::RadioButton
-		);
-
-		MenuBuilder.AddMenuEntry
-		(
-			LOCTEXT("ColorSize", "By Size"),
-			LOCTEXT("ColorSize_Tooltip", "Larger resources are more brightly colored."),
-			FSlateIcon(),
-			FUIAction(
-				FExecuteAction::CreateLambda([this]()
-				{
-					ResourceColor = EResourceColor::Size;
-					SetDirtyFlag();
-				}),
-				FCanExecuteAction(),
-				FIsActionChecked::CreateLambda([this]() { return ResourceColor == EResourceColor::Size; })
-			),
-			NAME_None,
-			EUserInterfaceActionType::RadioButton
-		);
-
-		MenuBuilder.EndSection();
-		MenuBuilder.BeginSection("FilterText", LOCTEXT("FilterTextHeader", "Track Resource Filter"));
-
-		MenuBuilder.AddWidget(
-			SNew(SVerticalBox)
-			+ SVerticalBox::Slot()
-			.AutoHeight()
-			.Padding(5.f)
-			[
-				// Search box allows for filtering
-				SNew(SSearchBox)
-				.InitialText(FText::FromString(FilterText))
-				.HintText(LOCTEXT("SearchHint", "Filter By Name"))
-				.OnTextChanged_Lambda([this](const FText& InText) { FilterText = InText.ToString(); SetDirtyFlag(); })
-			]
-			+ SVerticalBox::Slot()
-			.AutoHeight()
-			.Padding(5.f)
-			[
-				SNew(SHorizontalBox)
-				+ SHorizontalBox::Slot()
-				.VAlign(VAlign_Center)
-				[
-					SNew(STextBlock)
-					.Text(LOCTEXT("SizeThreshold", "Filter By Size (MB)"))
-				]
-				+ SHorizontalBox::Slot()
-				.AutoWidth()
-				.VAlign(VAlign_Center)
-				[
-					SNew(SSpinBox<float>)
-					.MinValue(0)
-					.MaxValue(1024)
-					.Value(FilterSize)
-					.MaxFractionalDigits(3)
-					.MinDesiredWidth(60)
-					.OnValueCommitted_Lambda([this](float InValue, ETextCommit::Type) { FilterSize = InValue; SetDirtyFlag(); })
-				]
-			],
-			FText::GetEmpty(),
-			true
-		);
-	
 		MenuBuilder.EndSection();
 	}
-	else if (Visualizer == EVisualizer::TransientHeaps)
+	else if (Visualizer == EVisualizer::TransientMemory)
 	{
 		const uint64 KiB = 1024;
 		const uint64 MiB = 1024 * KiB;
@@ -1861,11 +2006,11 @@ void FRenderGraphTrack::BuildContextMenu(FMenuBuilder& MenuBuilder)
 			FUIAction(
 				FExecuteAction::CreateLambda([this, MiB]()
 				{
-					TransientHeapBytesPerDepthSlot = 4 * MiB;
+					TransientBytesPerDepthSlot = 4 * MiB;
 					SetDirtyFlag();
 				}),
 				FCanExecuteAction(),
-				FIsActionChecked::CreateLambda([this, MiB]() { return TransientHeapBytesPerDepthSlot == 4 * MiB; })
+				FIsActionChecked::CreateLambda([this, MiB]() { return TransientBytesPerDepthSlot == 4 * MiB; })
 			),
 			NAME_None,
 			EUserInterfaceActionType::RadioButton
@@ -1879,11 +2024,11 @@ void FRenderGraphTrack::BuildContextMenu(FMenuBuilder& MenuBuilder)
 			FUIAction(
 				FExecuteAction::CreateLambda([this, MiB]()
 				{
-					TransientHeapBytesPerDepthSlot = 2 * MiB;
+					TransientBytesPerDepthSlot = 2 * MiB;
 					SetDirtyFlag();
 				}),
 				FCanExecuteAction(),
-				FIsActionChecked::CreateLambda([this, MiB]() { return TransientHeapBytesPerDepthSlot == 2 * MiB; })
+				FIsActionChecked::CreateLambda([this, MiB]() { return TransientBytesPerDepthSlot == 2 * MiB; })
 			),
 			NAME_None,
 			EUserInterfaceActionType::RadioButton
@@ -1897,11 +2042,11 @@ void FRenderGraphTrack::BuildContextMenu(FMenuBuilder& MenuBuilder)
 			FUIAction(
 				FExecuteAction::CreateLambda([this, MiB]()
 				{
-					TransientHeapBytesPerDepthSlot = MiB;
+					TransientBytesPerDepthSlot = MiB;
 					SetDirtyFlag();
 				}),
 				FCanExecuteAction(),
-				FIsActionChecked::CreateLambda([this, MiB]() { return TransientHeapBytesPerDepthSlot == MiB; })
+				FIsActionChecked::CreateLambda([this, MiB]() { return TransientBytesPerDepthSlot == MiB; })
 			),
 			NAME_None,
 			EUserInterfaceActionType::RadioButton
@@ -1915,11 +2060,11 @@ void FRenderGraphTrack::BuildContextMenu(FMenuBuilder& MenuBuilder)
 			FUIAction(
 				FExecuteAction::CreateLambda([this, KiB]()
 				{
-					TransientHeapBytesPerDepthSlot = 512 * KiB;
+					TransientBytesPerDepthSlot = 512 * KiB;
 					SetDirtyFlag();
 				}),
 				FCanExecuteAction(),
-				FIsActionChecked::CreateLambda([this, KiB]() { return TransientHeapBytesPerDepthSlot == 512 * KiB; })
+				FIsActionChecked::CreateLambda([this, KiB]() { return TransientBytesPerDepthSlot == 512 * KiB; })
 			),
 			NAME_None,
 			EUserInterfaceActionType::RadioButton
@@ -1933,11 +2078,11 @@ void FRenderGraphTrack::BuildContextMenu(FMenuBuilder& MenuBuilder)
 			FUIAction(
 				FExecuteAction::CreateLambda([this, KiB]()
 				{
-					TransientHeapBytesPerDepthSlot = 256 * KiB;
+					TransientBytesPerDepthSlot = 256 * KiB;
 					SetDirtyFlag();
 				}),
 				FCanExecuteAction(),
-				FIsActionChecked::CreateLambda([this, KiB]() { return TransientHeapBytesPerDepthSlot == 256 * KiB; })
+				FIsActionChecked::CreateLambda([this, KiB]() { return TransientBytesPerDepthSlot == 256 * KiB; })
 			),
 			NAME_None,
 			EUserInterfaceActionType::RadioButton
@@ -1951,11 +2096,11 @@ void FRenderGraphTrack::BuildContextMenu(FMenuBuilder& MenuBuilder)
 			FUIAction(
 				FExecuteAction::CreateLambda([this, KiB]()
 				{
-					TransientHeapBytesPerDepthSlot = 128 * KiB;
+					TransientBytesPerDepthSlot = 128 * KiB;
 					SetDirtyFlag();
 				}),
 				FCanExecuteAction(),
-				FIsActionChecked::CreateLambda([this, KiB]() { return TransientHeapBytesPerDepthSlot == 128 * KiB; })
+				FIsActionChecked::CreateLambda([this, KiB]() { return TransientBytesPerDepthSlot == 128 * KiB; })
 			),
 			NAME_None,
 			EUserInterfaceActionType::RadioButton
@@ -1969,18 +2114,84 @@ void FRenderGraphTrack::BuildContextMenu(FMenuBuilder& MenuBuilder)
 			FUIAction(
 				FExecuteAction::CreateLambda([this, KiB]()
 				{
-					TransientHeapBytesPerDepthSlot = 64 * KiB;
+					TransientBytesPerDepthSlot = 64 * KiB;
 					SetDirtyFlag();
 				}),
 				FCanExecuteAction(),
-				FIsActionChecked::CreateLambda([this, KiB]() { return TransientHeapBytesPerDepthSlot == 64 * KiB; })
+				FIsActionChecked::CreateLambda([this, KiB]() { return TransientBytesPerDepthSlot == 64 * KiB; })
 			),
 			NAME_None,
 			EUserInterfaceActionType::RadioButton
 		);
 
 		MenuBuilder.EndSection();
+
+		MenuBuilder.BeginSection("Show", LOCTEXT("ShowMenuHeader", "Track Show Flags"));
+
+		MenuBuilder.AddMenuEntry
+		(
+			LOCTEXT("ShowTracked", "Show Tracked"),
+			LOCTEXT("ShowTracked_Tooltip", "Show Tracked resources in the lifetime view."),
+			FSlateIcon(),
+			FUIAction(
+				FExecuteAction::CreateLambda([this]()
+				{
+					ResourceShow ^= EResourceShow::Tracked;
+					SetDirtyFlag();
+				}),
+				FCanExecuteAction(),
+				FIsActionChecked::CreateLambda([this]() { return EnumHasAnyFlags(ResourceShow, EResourceShow::Tracked); })
+			),
+			NAME_None,
+			EUserInterfaceActionType::ToggleButton
+		);
+
+		MenuBuilder.EndSection();
 	}
+	
+	MenuBuilder.BeginSection("FilterText", LOCTEXT("FilterTextHeader", "Track Resource Filter"));
+
+	MenuBuilder.AddWidget(
+		SNew(SVerticalBox)
+		+ SVerticalBox::Slot()
+		.AutoHeight()
+		.Padding(5.f)
+		[
+			// Search box allows for filtering
+			SNew(SSearchBox)
+			.InitialText(FText::FromString(FilterText))
+			.HintText(LOCTEXT("SearchHint", "Filter By Name"))
+			.OnTextChanged_Lambda([this](const FText& InText) { FilterText = InText.ToString(); SetDirtyFlag(); })
+		]
+		+ SVerticalBox::Slot()
+		.AutoHeight()
+		.Padding(5.f)
+		[
+			SNew(SHorizontalBox)
+			+ SHorizontalBox::Slot()
+			.VAlign(VAlign_Center)
+			[
+				SNew(STextBlock)
+				.Text(LOCTEXT("SizeThreshold", "Filter By Size (MB)"))
+			]
+			+ SHorizontalBox::Slot()
+			.AutoWidth()
+			.VAlign(VAlign_Center)
+			[
+				SNew(SSpinBox<float>)
+				.MinValue(0)
+				.MaxValue(1024)
+				.Value(FilterSize)
+				.MaxFractionalDigits(3)
+				.MinDesiredWidth(60)
+				.OnValueCommitted_Lambda([this](float InValue, ETextCommit::Type) { FilterSize = InValue; SetDirtyFlag(); })
+			]
+		],
+		FText::GetEmpty(),
+		true
+	);
+	
+	MenuBuilder.EndSection();
 }
 
 void FRenderGraphTrack::InitTooltip(FTooltipDrawState& Tooltip, const ITimingEvent& InTooltipEvent) const
@@ -1991,11 +2202,24 @@ void FRenderGraphTrack::InitTooltip(FTooltipDrawState& Tooltip, const ITimingEve
 	{
 		const auto AddCommonResourceText = [&](const FResourcePacket& Resource)
 		{
-			if (Resource.TransientStats.HeapOffsetMax > 0)
+			if (Resource.bTransient)
 			{
-				Tooltip.AddNameValueTextLine(TEXT("Heap Index:"), FString::Printf(TEXT("%u"), Resource.TransientStats.HeapIndex));
-				Tooltip.AddNameValueTextLine(TEXT("Heap Start:"), FString::Printf(TEXT("%.3fMB"), (double)Resource.TransientStats.HeapOffsetMin / (1024.0 * 1024.0)));
-				Tooltip.AddNameValueTextLine(TEXT("Heap End:"), FString::Printf(TEXT("%.3fMB"), (double)Resource.TransientStats.HeapOffsetMax / (1024.0 * 1024.0)));
+				Tooltip.AddTextLine(TEXT("Transient"), FLinearColor::Red);
+				Tooltip.AddNameValueTextLine(TEXT("Transient Cache Hit:"), Resource.bTransientCacheHit ? TEXT("Yes") : TEXT("No"));
+			}
+
+			if (!Resource.TransientAllocations.IsEmpty())
+			{
+				for (const FRHITransientAllocationStats::FAllocation& Allocation : Resource.TransientAllocations)
+				{
+					Tooltip.AddNameValueTextLine(TEXT("Transient Memory:"),
+						FString::Printf(TEXT("Range: %u, Start: %.3fMb, End: %.3fMb"),
+							Allocation.MemoryRangeIndex,
+							(float)Allocation.OffsetMin / (1024.0f * 1024.0f),
+							(float)Allocation.OffsetMax / (1024.0f * 1024.0f)
+						)
+					);
+				}
 			}
 
 			if (Resource.bExtracted)
@@ -2008,9 +2232,9 @@ void FRenderGraphTrack::InitTooltip(FTooltipDrawState& Tooltip, const ITimingEve
 				Tooltip.AddTextLine(TEXT("External"), FLinearColor::Red);
 			}
 
-			if (Resource.bTransient)
+			if (Resource.bTransientUntracked)
 			{
-				Tooltip.AddTextLine(TEXT("Transient"), FLinearColor::Red);
+				Tooltip.AddTextLine(TEXT("Untracked by RDG"), FLinearColor::Red);
 			}
 		};
 
diff --git a/Engine/Plugins/RenderGraphInsights/Source/RenderGraphInsights/Private/RenderGraphTrack.h b/Engine/Plugins/RenderGraphInsights/Source/RenderGraphInsights/Private/RenderGraphTrack.h
--- a/Engine/Plugins/RenderGraphInsights/Source/RenderGraphInsights/Private/RenderGraphTrack.h
+++ b/Engine/Plugins/RenderGraphInsights/Source/RenderGraphInsights/Private/RenderGraphTrack.h
@@ -115,8 +115,28 @@ public:
 };
 
 using FVisibleScope   = TVisibleItemHelper<FScopePacket>;
-using FVisibleTexture = TVisibleItemHelper<FTexturePacket>;
-using FVisibleBuffer  = TVisibleItemHelper<FBufferPacket>;
+
+template <typename InPacketType>
+class TVisibleResourceHelper : public FVisibleItem
+{
+public:
+	using PacketType = InPacketType;
+
+	TVisibleResourceHelper(const FTimingTrackViewport& Viewport, const PacketType& InPacket, uint32 InColor, double InStartTime, double InEndTime, float InDepthY, float InDepthH = 1.0f, uint32 InAllocationIndex = 0)
+		: FVisibleItem(Viewport, InPacket, InColor, InStartTime, InEndTime, InDepthY, InDepthH)
+		, AllocationIndex(InAllocationIndex)
+	{}
+
+	const PacketType& GetPacket() const
+	{
+		return static_cast<const PacketType&>(*Packet);
+	}
+
+	uint32 AllocationIndex;
+};
+
+using FVisibleTexture = TVisibleResourceHelper<FTexturePacket>;
+using FVisibleBuffer  = TVisibleResourceHelper<FBufferPacket>;
 
 class FVisibleGraph final : public TVisibleItemHelper<FGraphPacket>
 {
@@ -144,24 +164,22 @@ public:
 		return Passes[Pass.VisibleIndex];
 	}
 
-	const FVisibleTexture* GetVisibleTexture(const FTexturePacket& Texture) const
+	template <typename FunctionType>
+	void EnumerateVisibleTextures(const FTexturePacket& Texture, FunctionType Function) const
 	{
-		if (Texture.VisibleIndex == kInvalidVisibleIndex)
+		for (int32 VisibleIndex : Texture.VisibleItems)
 		{
-			return nullptr;
+			Function(Textures[VisibleIndex]);
 		}
-
-		return &Textures[Texture.VisibleIndex];
 	}
 
-	const FVisibleBuffer* GetVisibleBuffer(const FBufferPacket& Buffer) const
+	template <typename FunctionType>
+	void EnumerateVisibleBuffers(const FBufferPacket& Buffer, FunctionType Function) const
 	{
-		if (Buffer.VisibleIndex == kInvalidVisibleIndex)
+		for (int32 VisibleIndex : Buffer.VisibleItems)
 		{
-			return nullptr;
+			Function(Buffers[VisibleIndex]);
 		}
-
-		return &Buffers[Buffer.VisibleIndex];
 	}
 
 	void Draw(FRenderGraphTrackDrawStateBuilder& Builder) const override;
@@ -176,7 +194,7 @@ public:
 	TArray<FVisibleTexture> Textures;
 	TArray<FVisibleBuffer> Buffers;
 	TArray<uint32> AsyncComputePasses;
-	TArray<float, TInlineAllocator<8>> TransientHeapDepths;
+	TArray<float> TransientMemoryRangeDepthOffsets;
 	float HeaderEndDepth{};
 	float HeaderPassBarDepth{};
 };
@@ -213,13 +231,13 @@ private:
 	{
 		const uint32 VisibleIndex = VisibleGraphs.Num();
 		FVisibleGraph& VisibleGraph = VisibleGraphs.Emplace_GetRef(Viewport, Graph, Color, 1.0f);
-		check(VisibleGraph.Packet->VisibleIndex == kInvalidVisibleIndex);
-		VisibleGraph.Packet->VisibleIndex = VisibleIndex;
+		check(VisibleGraph.GetPacket().VisibleIndex == kInvalidVisibleIndex);
+		VisibleGraph.GetPacket().VisibleIndex = VisibleIndex;
 		VisibleGraph.Index = VisibleIndex;
 
-		for (uint64 HeapOffset : Graph.TransientHeapOffsets)
+		for (uint64 ByteOffset : Graph.TransientMemoryRangeByteOffsets)
 		{
-			VisibleGraph.TransientHeapDepths.Emplace(FMath::CeilToFloat(TransientHeapOffsetToDepth(HeapOffset)));
+			VisibleGraph.TransientMemoryRangeDepthOffsets.Emplace(FMath::CeilToFloat(TransientByteOffsetToDepth(ByteOffset)));
 		}
 
 		return VisibleGraph;
@@ -234,9 +252,9 @@ private:
 		return nullptr;
 	}
 
-	float TransientHeapOffsetToDepth(uint64 HeapOffset) const;
-	float TransientHeapOffsetToDepthAligned(uint64 HeapOffset) const;
-	float TransientHeapOffsetToSlateY(const FTimingViewLayout& Layout, uint64 HeapOffset) const;
+	float TransientByteOffsetToDepth(uint64 ByteOffset) const;
+	float TransientByteOffsetToDepthAligned(uint64 ByteOffset) const;
+	float TransientByteOffsetToSlateY(const FTimingViewLayout& Layout, uint64 ByteOffset) const;
 
 	float MouseSlateX{};
 	float MouseSlateY{};
@@ -263,7 +281,10 @@ private:
 	bool ShowBuffers() const	{ return EnumHasAnyFlags(ResourceShow, EResourceShow::Buffers); }
 	bool ShowTransient() const	{ return EnumHasAnyFlags(ResourceShow, EResourceShow::Transient); }
 	bool ShowExternal() const	{ return EnumHasAnyFlags(ResourceShow, EResourceShow::External); }
-	bool ShowPooled() const		{ return EnumHasAnyFlags(ResourceShow, EResourceShow::Pooled);}
+	bool ShowInternal() const	{ return EnumHasAnyFlags(ResourceShow, EResourceShow::Internal); }
+	bool ShowExtracted() const	{ return EnumHasAnyFlags(ResourceShow, EResourceShow::Extracted); }
+	bool ShowPooled() const		{ return EnumHasAnyFlags(ResourceShow, EResourceShow::Pooled); }
+	bool ShowTracked() const	{ return EnumHasAnyFlags(ResourceShow, EResourceShow::Tracked); }
 
 	enum class EResourceShow
 	{
@@ -271,8 +292,11 @@ private:
 		Buffers			= 1 << 1,
 		Transient		= 1 << 2,
 		External		= 1 << 3,
-		Pooled			= 1 << 4,
-		All				= Textures | Buffers | Transient | External | Pooled
+		Internal		= 1 << 4,
+		Extracted		= 1 << 5,
+		Pooled			= 1 << 6,
+		Tracked			= 1 << 7,
+		All				= Textures | Buffers | Transient | External | Internal | Extracted | Pooled | Tracked
 	};
 	FRIEND_ENUM_CLASS_FLAGS(FRenderGraphTrack::EResourceShow);
 
@@ -288,13 +312,14 @@ private:
 	enum class EResourceColor
 	{
 		Type,
-		Size
+		Size,
+		TransientCache
 	};
 
 	enum class EVisualizer
 	{
 		Resources,
-		TransientHeaps
+		TransientMemory
 	};
 
 	EResourceShow ResourceShow = EResourceShow::All;
@@ -303,7 +328,7 @@ private:
 	EVisualizer Visualizer = EVisualizer::Resources;
 	FString FilterText;
 	float FilterSize{};
-	uint64 TransientHeapBytesPerDepthSlot = 1024 * 1024;
+	uint64 TransientBytesPerDepthSlot = 1024 * 1024;
 
 	FTooltipDrawState SelectedTooltipState;
 
diff --git a/Engine/Source/Runtime/D3D12RHI/Private/D3D12Adapter.cpp b/Engine/Source/Runtime/D3D12RHI/Private/D3D12Adapter.cpp
--- a/Engine/Source/Runtime/D3D12RHI/Private/D3D12Adapter.cpp
+++ b/Engine/Source/Runtime/D3D12RHI/Private/D3D12Adapter.cpp
@@ -817,6 +817,16 @@ void FD3D12Adapter::CreateRootDevice(bool bWithDebug)
 #endif
 }
 
+FD3D12TransientHeapCache& FD3D12Adapter::GetOrCreateTransientHeapCache()
+{
+	if (!TransientMemoryCache)
+	{
+		TransientMemoryCache = FD3D12TransientHeapCache::Create(this, FRHIGPUMask::All());
+	}
+
+	return static_cast<FD3D12TransientHeapCache&>(*TransientMemoryCache);
+}
+
 void FD3D12Adapter::InitializeDevices()
 {
 	check(IsInGameThread());
@@ -1134,8 +1144,6 @@ void FD3D12Adapter::InitializeDevices()
 
 		PipelineStateCache.RebuildFromDiskCache(StaticGraphicsRS, StaticComputeRS);
 	}
-
-	TransientResourceSystem = FD3D12TransientResourceSystem::Create(this, FRHIGPUMask::All());
 }
 
 void FD3D12Adapter::InitializeRayTracing()
@@ -1185,8 +1193,6 @@ void FD3D12Adapter::Cleanup()
 	// Reset the RHI initialized flag.
 	GIsRHIInitialized = false;
 
-	TransientResourceSystem.Reset();
-
 	for (auto& Viewport : Viewports)
 	{
 		Viewport->IssueFrameEvent();
@@ -1256,6 +1262,7 @@ void FD3D12Adapter::Cleanup()
 		StagingFence.SafeRelease();
 	}
 
+	TransientMemoryCache.Reset();
 
 	PipelineStateCache.Close();
 	RootSignatureManager.Destroy();
@@ -1396,8 +1403,10 @@ void FD3D12Adapter::EndFrame()
 	}
 	GetDeferredDeletionQueue().ReleaseResources(false, false);
 
-	TransientResourceSystem->GarbageCollect();
-	TransientResourceSystem->UpdateStats();
+	if (TransientMemoryCache)
+	{
+		TransientMemoryCache->GarbageCollect();
+	}
 
 #if D3D12_SUBMISSION_GAP_RECORDER
 	SubmitGapRecorderTimestamps();
diff --git a/Engine/Source/Runtime/D3D12RHI/Private/D3D12Adapter.h b/Engine/Source/Runtime/D3D12RHI/Private/D3D12Adapter.h
--- a/Engine/Source/Runtime/D3D12RHI/Private/D3D12Adapter.h
+++ b/Engine/Source/Runtime/D3D12RHI/Private/D3D12Adapter.h
@@ -377,7 +377,7 @@ public:
 	inline FD3D12CommandContextRedirector& GetDefaultContextRedirector() { return DefaultContextRedirector; }
 	inline FD3D12CommandContextRedirector& GetDefaultAsyncComputeContextRedirector() { return DefaultAsyncComputeContextRedirector; }
 
-	inline FD3D12TransientResourceSystem& GetTransientResourceSystem() { return *TransientResourceSystem; }
+	FD3D12TransientHeapCache& GetOrCreateTransientHeapCache();
 
 	FD3D12TemporalEffect* GetTemporalEffect(const FName& EffectName);
 
@@ -571,7 +571,7 @@ protected:
 	TMap<FName, FD3D12TemporalEffect> TemporalEffectMap;
 #endif
 
-	TUniquePtr<FD3D12TransientResourceSystem> TransientResourceSystem;
+	TUniquePtr<IRHITransientMemoryCache> TransientMemoryCache;
 
 	// Each of these devices represents a physical GPU 'Node'
 	FD3D12Device* Devices[MAX_NUM_GPUS];
diff --git a/Engine/Source/Runtime/D3D12RHI/Private/D3D12CommandContext.cpp b/Engine/Source/Runtime/D3D12RHI/Private/D3D12CommandContext.cpp
--- a/Engine/Source/Runtime/D3D12RHI/Private/D3D12CommandContext.cpp
+++ b/Engine/Source/Runtime/D3D12RHI/Private/D3D12CommandContext.cpp
@@ -1027,7 +1027,7 @@ void FD3D12DynamicRHI::RHIReleaseTransition(FRHITransition* Transition)
 
 IRHITransientResourceAllocator* FD3D12DynamicRHI::RHICreateTransientResourceAllocator()
 {
-	return new FD3D12TransientResourceAllocator(GetAdapter().GetTransientResourceSystem());
+	return new FD3D12TransientResourceHeapAllocator(GetAdapter().GetOrCreateTransientHeapCache());
 }
 
 //////////////////////////////////////////////////////////////////////////////////////////////////////////
diff --git a/Engine/Source/Runtime/D3D12RHI/Private/D3D12TransientResourceAllocator.cpp b/Engine/Source/Runtime/D3D12RHI/Private/D3D12TransientResourceAllocator.cpp
--- a/Engine/Source/Runtime/D3D12RHI/Private/D3D12TransientResourceAllocator.cpp
+++ b/Engine/Source/Runtime/D3D12RHI/Private/D3D12TransientResourceAllocator.cpp
@@ -24,7 +24,7 @@ D3D12_RESOURCE_STATES GetInitialResourceState(const D3D12_RESOURCE_DESC& InDesc)
 	return State;
 }
 
-FD3D12TransientHeap::FD3D12TransientHeap(const FRHITransientHeapInitializer& Initializer, FD3D12Adapter* Adapter, FD3D12Device* Device, FRHIGPUMask VisibleNodeMask)
+FD3D12TransientHeap::FD3D12TransientHeap(const FInitializer& Initializer, FD3D12Adapter* Adapter, FD3D12Device* Device, FRHIGPUMask VisibleNodeMask)
 	: FRHITransientHeap(Initializer)
 {
 	D3D12_HEAP_FLAGS HeapFlags = D3D12_HEAP_FLAG_ALLOW_ALL_BUFFERS_AND_TEXTURES;
@@ -54,7 +54,7 @@ FD3D12TransientHeap::FD3D12TransientHeap(const FRHITransientHeapInitializer& Ini
 	D3D12_HEAP_DESC Desc = {};
 	Desc.SizeInBytes = Initializer.Size;
 	Desc.Properties = HeapProperties;
-	Desc.Alignment = Initializer.Alignment;
+	Desc.Alignment = D3D12_DEFAULT_RESOURCE_PLACEMENT_ALIGNMENT;
 	Desc.Flags = HeapFlags;
 
 	if (Adapter->IsHeapNotZeroedSupported())
@@ -86,7 +86,7 @@ FD3D12TransientHeap::FD3D12TransientHeap(const FRHITransientHeapInitializer& Ini
 	Heap->SetHeap(D3DHeap, TEXT("TransientResourceAllocator Backing Heap"), true, true);
 	Heap->BeginTrackingResidency(Desc.SizeInBytes);
 
-	BaseGPUVirtualAddress = Heap->GetGPUVirtualAddress();
+	SetGPUVirtualAddress(Heap->GetGPUVirtualAddress());
 
 	INC_MEMORY_STAT_BY(STAT_D3D12TransientHeaps, Desc.SizeInBytes);
 }
@@ -100,24 +100,23 @@ FD3D12TransientHeap::~FD3D12TransientHeap()
 	}
 }
 
-TUniquePtr<FD3D12TransientResourceSystem> FD3D12TransientResourceSystem::Create(FD3D12Adapter* ParentAdapter, FRHIGPUMask VisibleNodeMask)
+TUniquePtr<FD3D12TransientHeapCache> FD3D12TransientHeapCache::Create(FD3D12Adapter* ParentAdapter, FRHIGPUMask VisibleNodeMask)
 {
-	FRHITransientResourceSystemInitializer Initializer = FRHITransientResourceSystemInitializer::CreateDefault();
-	Initializer.HeapAlignment = D3D12_DEFAULT_RESOURCE_PLACEMENT_ALIGNMENT;
+	FRHITransientHeapCache::FInitializer Initializer = FRHITransientHeapCache::FInitializer::CreateDefault();
 
 	// Tier2 hardware is able to mix resource types onto the same heap.
 	Initializer.bSupportsAllHeapFlags = ParentAdapter->GetResourceHeapTier() == D3D12_RESOURCE_HEAP_TIER_2;
 
-	return TUniquePtr<FD3D12TransientResourceSystem>(new FD3D12TransientResourceSystem(Initializer, ParentAdapter, VisibleNodeMask));
+	return TUniquePtr<FD3D12TransientHeapCache>(new FD3D12TransientHeapCache(Initializer, ParentAdapter, VisibleNodeMask));
 }
 
-FD3D12TransientResourceSystem::FD3D12TransientResourceSystem(const FRHITransientResourceSystemInitializer& Initializer, FD3D12Adapter* ParentAdapter, FRHIGPUMask InVisibleNodeMask)
-	: FRHITransientResourceSystem(Initializer)
+FD3D12TransientHeapCache::FD3D12TransientHeapCache(const FRHITransientHeapCache::FInitializer& Initializer, FD3D12Adapter* ParentAdapter, FRHIGPUMask InVisibleNodeMask)
+	: FRHITransientHeapCache(Initializer)
 	, FD3D12AdapterChild(ParentAdapter)
 	, VisibleNodeMask(InVisibleNodeMask)
 {}
 
-FRHITransientHeap* FD3D12TransientResourceSystem::CreateHeap(const FRHITransientHeapInitializer& HeapInitializer)
+FRHITransientHeap* FD3D12TransientHeapCache::CreateHeap(const FRHITransientHeap::FInitializer& HeapInitializer)
 {
 	return GetParentAdapter()->CreateLinkedObject<FD3D12TransientHeap>(VisibleNodeMask, [&](FD3D12Device* Device)
 	{
@@ -125,21 +124,23 @@ FRHITransientHeap* FD3D12TransientResourceSystem::CreateHeap(const FRHITransient
 	});
 }
 
-FD3D12TransientResourceAllocator::FD3D12TransientResourceAllocator(FD3D12TransientResourceSystem& InParentSystem)
-	: FD3D12AdapterChild(InParentSystem.GetParentAdapter())
-	, Allocator(InParentSystem)
+FD3D12TransientResourceHeapAllocator::FD3D12TransientResourceHeapAllocator(FD3D12TransientHeapCache& InHeapCache)
+	: FRHITransientResourceHeapAllocator(InHeapCache)
+	, FD3D12AdapterChild(InHeapCache.GetParentAdapter())
 	, AllocationInfoQueryDevice(GetParentAdapter()->GetDevice(0))
 {}
 
-FRHITransientTexture* FD3D12TransientResourceAllocator::CreateTexture(const FRHITextureCreateInfo& InCreateInfo, const TCHAR* InDebugName, uint32 InPassIndex)
+FRHITransientTexture* FD3D12TransientResourceHeapAllocator::CreateTexture(const FRHITextureCreateInfo& InCreateInfo, const TCHAR* InDebugName, uint32 InPassIndex)
 {
 	FD3D12DynamicRHI* DynamicRHI = FD3D12DynamicRHI::GetD3DRHI();
 
 	const D3D12_RESOURCE_DESC Desc = DynamicRHI->GetResourceDesc(InCreateInfo);
-	const D3D12_RESOURCE_ALLOCATION_INFO Info = AllocationInfoQueryDevice->GetResourceAllocationInfo(Desc);
+	D3D12_RESOURCE_ALLOCATION_INFO Info = AllocationInfoQueryDevice->GetResourceAllocationInfo(Desc);
 
-	return Allocator.CreateTexture(InCreateInfo, InDebugName, InPassIndex, Info.SizeInBytes, Info.Alignment,
-		[&](const FRHITransientResourceAllocator::FResourceInitializer& Initializer)
+	Info.Alignment = FMath::Max<uint32>(D3D12_DEFAULT_RESOURCE_PLACEMENT_ALIGNMENT, Info.Alignment);
+
+	return CreateTextureInternal(InCreateInfo, InDebugName, InPassIndex, Info.SizeInBytes, Info.Alignment,
+		[&](const FRHITransientHeap::FResourceInitializer& Initializer)
 	{
 		ERHIAccess InitialState = ERHIAccess::UAVMask;
 
@@ -154,13 +155,12 @@ FRHITransientTexture* FD3D12TransientResourceAllocator::CreateTexture(const FRHI
 
 		ED3D12ResourceTransientMode TransientMode = ED3D12ResourceTransientMode::Transient;
 		FResourceAllocatorAdapter ResourceAllocatorAdapter(GetParentAdapter(), static_cast<FD3D12TransientHeap&>(Initializer.Heap), Initializer.Allocation, Desc);
-		
 		FRHITexture* Texture = DynamicRHI->CreateTexture(InCreateInfo, InDebugName, InitialState, TransientMode, &ResourceAllocatorAdapter);
-		return new FRHITransientTexture(Texture, Initializer.Hash, InCreateInfo);
+		return new FRHITransientTexture(Texture, ResourceAllocatorAdapter.GpuVirtualAddress, Initializer.Hash, Info.SizeInBytes, ERHITransientAllocationType::Heap, InCreateInfo);
 	});
 }
 
-void FD3D12TransientResourceAllocator::FResourceAllocatorAdapter::AllocateResource(
+void FD3D12TransientResourceHeapAllocator::FResourceAllocatorAdapter::AllocateResource(
 	uint32 GPUIndex, D3D12_HEAP_TYPE, const FD3D12ResourceDesc& InDesc, uint64 InSize, uint32, ED3D12ResourceStateMode InResourceStateMode,
 	D3D12_RESOURCE_STATES InCreateState, const D3D12_CLEAR_VALUE* InClearValue, const TCHAR* InName, FD3D12ResourceLocation& ResourceLocation)
 {
@@ -182,6 +182,8 @@ void FD3D12TransientResourceAllocator::FResourceAllocatorAdapter::AllocateResour
 	ResourceLocation.SetSize(InSize);
 	ResourceLocation.SetTransient(true);
 
+	GpuVirtualAddress = NewResource->GetGPUVirtualAddress();
+
 #if TRACK_RESOURCE_ALLOCATIONS
 	if (Adapter->IsTrackingAllAllocations())
 	{
@@ -191,7 +193,7 @@ void FD3D12TransientResourceAllocator::FResourceAllocatorAdapter::AllocateResour
 #endif
 }
 
-FRHITransientBuffer* FD3D12TransientResourceAllocator::CreateBuffer(const FRHIBufferCreateInfo& InCreateInfo, const TCHAR* InDebugName, uint32 InPassIndex)
+FRHITransientBuffer* FD3D12TransientResourceHeapAllocator::CreateBuffer(const FRHIBufferCreateInfo& InCreateInfo, const TCHAR* InDebugName, uint32 InPassIndex)
 {
 	D3D12_RESOURCE_DESC Desc;
 	uint32 Alignment;
@@ -201,27 +203,12 @@ FRHITransientBuffer* FD3D12TransientResourceAllocator::CreateBuffer(const FRHIBu
 	Alignment = D3D12_DEFAULT_RESOURCE_PLACEMENT_ALIGNMENT;
 	uint64 Size = Align(Desc.Width, Alignment);
 
-	return Allocator.CreateBuffer(InCreateInfo, InDebugName, InPassIndex, Size, Alignment,
-		[&](const FRHITransientResourceAllocator::FResourceInitializer& Initializer)
+	return CreateBufferInternal(InCreateInfo, InDebugName, InPassIndex, Size, Alignment,
+		[&](const FRHITransientHeap::FResourceInitializer& Initializer)
 	{
 		ED3D12ResourceTransientMode TransientMode = ED3D12ResourceTransientMode::Transient;
 		FResourceAllocatorAdapter ResourceAllocatorAdapter(GetParentAdapter(), static_cast<FD3D12TransientHeap&>(Initializer.Heap), Initializer.Allocation, Desc);
 		FRHIBuffer* Buffer = FD3D12DynamicRHI::GetD3DRHI()->CreateBuffer(InCreateInfo, InDebugName, ERHIAccess::UAVMask, TransientMode, &ResourceAllocatorAdapter);
-		return new FRHITransientBuffer(Buffer, Initializer.Hash, InCreateInfo);
+		return new FRHITransientBuffer(Buffer, ResourceAllocatorAdapter.GpuVirtualAddress, Initializer.Hash, Size, ERHITransientAllocationType::Heap, InCreateInfo);
 	});
-}
-
-void FD3D12TransientResourceAllocator::DeallocateMemory(FRHITransientTexture* InTexture, uint32 InPassIndex)
-{
-	Allocator.DeallocateMemory(InTexture, InPassIndex);
-}
-
-void FD3D12TransientResourceAllocator::DeallocateMemory(FRHITransientBuffer* InBuffer, uint32 InPassIndex)
-{
-	Allocator.DeallocateMemory(InBuffer, InPassIndex);
-}
-
-void FD3D12TransientResourceAllocator::Freeze(FRHICommandListImmediate& RHICmdList, FRHITransientHeapStats& OutHeapStats)
-{
-	Allocator.Freeze(RHICmdList, OutHeapStats);
-}
+}
\ No newline at end of file
diff --git a/Engine/Source/Runtime/D3D12RHI/Private/D3D12TransientResourceAllocator.h b/Engine/Source/Runtime/D3D12RHI/Private/D3D12TransientResourceAllocator.h
--- a/Engine/Source/Runtime/D3D12RHI/Private/D3D12TransientResourceAllocator.h
+++ b/Engine/Source/Runtime/D3D12RHI/Private/D3D12TransientResourceAllocator.h
@@ -13,7 +13,7 @@ class FD3D12TransientHeap final
 	, public FD3D12LinkedAdapterObject<FD3D12TransientHeap>
 {
 public:
-	FD3D12TransientHeap(const FRHITransientHeapInitializer& Initializer, FD3D12Adapter* Adapter, FD3D12Device* Device, FRHIGPUMask VisibleNodeMask);
+	FD3D12TransientHeap(const FInitializer& Initializer, FD3D12Adapter* Adapter, FD3D12Device* Device, FRHIGPUMask VisibleNodeMask);
 	~FD3D12TransientHeap();
 
 	FD3D12Heap* Get() { return Heap; }
@@ -22,35 +22,32 @@ private:
 	TRefCountPtr<FD3D12Heap> Heap;
 };
 
-class FD3D12TransientResourceSystem final
-	: public FRHITransientResourceSystem
+class FD3D12TransientHeapCache final
+	: public FRHITransientHeapCache
 	, public FD3D12AdapterChild
 {
 public:
-	static TUniquePtr<FD3D12TransientResourceSystem> Create(FD3D12Adapter* ParentAdapter, FRHIGPUMask VisibleNodeMask);
+	static TUniquePtr<FD3D12TransientHeapCache> Create(FD3D12Adapter* ParentAdapter, FRHIGPUMask VisibleNodeMask);
 
 	//! FRHITransientResourceSystem Overrides
-	FRHITransientHeap* CreateHeap(const FRHITransientHeapInitializer& Initializer) override;
+	FRHITransientHeap* CreateHeap(const FRHITransientHeap::FInitializer& Initializer) override;
 
 private:
-	FD3D12TransientResourceSystem(const FRHITransientResourceSystemInitializer& Initializer, FD3D12Adapter* ParentAdapter, FRHIGPUMask VisibleNodeMask);
+	FD3D12TransientHeapCache(const FInitializer& Initializer, FD3D12Adapter* ParentAdapter, FRHIGPUMask VisibleNodeMask);
 
 	FRHIGPUMask VisibleNodeMask;
 };
 
-class FD3D12TransientResourceAllocator final
-	: public IRHITransientResourceAllocator
+class FD3D12TransientResourceHeapAllocator final
+	: public FRHITransientResourceHeapAllocator
 	, public FD3D12AdapterChild
 {
 public:
-	FD3D12TransientResourceAllocator(FD3D12TransientResourceSystem& InParentSystem);
+	FD3D12TransientResourceHeapAllocator(FD3D12TransientHeapCache& InHeapCache);
 
 	//! IRHITransientResourceAllocator Overrides
 	FRHITransientTexture* CreateTexture(const FRHITextureCreateInfo& InCreateInfo, const TCHAR* InDebugName, uint32 InPassIndex) override;
 	FRHITransientBuffer* CreateBuffer(const FRHIBufferCreateInfo& InCreateInfo, const TCHAR* InDebugName, uint32 InPassIndex) override;
-	void DeallocateMemory(FRHITransientTexture* InTexture, uint32 InPassIndex) override;
-	void DeallocateMemory(FRHITransientBuffer* InBuffer, uint32 InPassIndex) override;
-	void Freeze(FRHICommandListImmediate& RHICmdList, FRHITransientHeapStats& OutHeapStats) override;
 
 private:
 
@@ -72,8 +69,8 @@ private:
 		FD3D12TransientHeap& Heap;
 		const FRHITransientHeapAllocation& Allocation;
 		const D3D12_RESOURCE_DESC Desc;
+		uint64 GpuVirtualAddress = 0;
 	};
 
-	FRHITransientResourceAllocator Allocator;
 	FD3D12Device* AllocationInfoQueryDevice;
 };
\ No newline at end of file
diff --git a/Engine/Source/Runtime/RHI/Private/RHI.cpp b/Engine/Source/Runtime/RHI/Private/RHI.cpp
--- a/Engine/Source/Runtime/RHI/Private/RHI.cpp
+++ b/Engine/Source/Runtime/RHI/Private/RHI.cpp
@@ -2753,9 +2753,9 @@ void FRHIBufferViewCache::SetDebugName(const TCHAR* DebugName)
 
 #endif
 
-void FRHITransientTexture::Init(const TCHAR* InName, uint32 InAllocationIndex, uint32 InPassIndex, const FRHITransientResourceStats& InStats)
+void FRHITransientTexture::Acquire(const TCHAR* InName, uint32 InPassIndex, uint64 InAcquireCycle)
 {
-	FRHITransientResource::Init(InName, InAllocationIndex, InPassIndex, InStats);
+	FRHITransientResource::Acquire(InName, InPassIndex, InAcquireCycle);
 	ViewCache.SetDebugName(InName);
 
 #if !(UE_BUILD_SHIPPING || UE_BUILD_TEST)
@@ -2763,9 +2763,9 @@ void FRHITransientTexture::Init(const TCHAR* InName, uint32 InAllocationIndex, u
 #endif
 }
 
-void FRHITransientBuffer::Init(const TCHAR* InName, uint32 InAllocationIndex, uint32 InPassIndex, const FRHITransientResourceStats& InStats)
+void FRHITransientBuffer::Acquire(const TCHAR* InName, uint32 InPassIndex, uint64 InAcquireCycle)
 {
-	FRHITransientResource::Init(InName, InAllocationIndex, InPassIndex, InStats);
+	FRHITransientResource::Acquire(InName, InPassIndex, InAcquireCycle);
 	ViewCache.SetDebugName(InName);
 
 	// TODO: Add method to rename a buffer.
diff --git a/Engine/Source/Runtime/RHI/Private/RHIValidation.cpp b/Engine/Source/Runtime/RHI/Private/RHIValidation.cpp
--- a/Engine/Source/Runtime/RHI/Private/RHIValidation.cpp
+++ b/Engine/Source/Runtime/RHI/Private/RHIValidation.cpp
@@ -1536,13 +1536,11 @@ TLockFreePointerListUnordered<FValidationContext, PLATFORM_CACHE_LINE_SIZE> FVal
 
 FValidationTransientResourceAllocator::~FValidationTransientResourceAllocator()
 {
-	check(bFrozen);
-	check(bReleased);
+	delete RHIAllocator;
 }
 
 FRHITransientTexture* FValidationTransientResourceAllocator::CreateTexture(const FRHITextureCreateInfo& InCreateInfo, const TCHAR* InDebugName, uint32 InPassIndex)
 {
-	check(!bFrozen);
 	checkf(!EnumHasAnyFlags(InCreateInfo.Flags, TexCreate_Transient), TEXT("Attempted to create transient texture %s using legacy TexCreate_Transient flag."), InDebugName);
 
 	FRHITransientTexture* TransientTexture = RHIAllocator->CreateTexture(InCreateInfo, InDebugName, InPassIndex);
@@ -1558,7 +1556,6 @@ FRHITransientTexture* FValidationTransientResourceAllocator::CreateTexture(const
 	FAllocatedResourceData ResourceData;
 	ResourceData.DebugName = InDebugName;
 	ResourceData.ResourceType = FAllocatedResourceData::EType::Texture;
-	ResourceData.bMemoryAllocated = true;
 	ResourceData.Texture.Flags = InCreateInfo.Flags;
 	ResourceData.Texture.Format = InCreateInfo.Format;
 	ResourceData.Texture.ArraySize = InCreateInfo.ArraySize;
@@ -1582,7 +1579,6 @@ FRHITransientTexture* FValidationTransientResourceAllocator::CreateTexture(const
 
 FRHITransientBuffer* FValidationTransientResourceAllocator::CreateBuffer(const FRHIBufferCreateInfo& InCreateInfo, const TCHAR* InDebugName, uint32 InPassIndex)
 {
-	check(!bFrozen);
 	checkf(!EnumHasAnyFlags(InCreateInfo.Usage, BUF_Transient), TEXT("Attempted to create transient buffer %s using legacy BUF_Transient flag."), InDebugName);
 
 	FRHITransientBuffer* TransientBuffer = RHIAllocator->CreateBuffer(InCreateInfo, InDebugName, InPassIndex);
@@ -1598,7 +1594,6 @@ FRHITransientBuffer* FValidationTransientResourceAllocator::CreateBuffer(const F
 	FAllocatedResourceData ResourceData;
 	ResourceData.DebugName = InDebugName;
 	ResourceData.ResourceType = FAllocatedResourceData::EType::Buffer;
-	ResourceData.bMemoryAllocated = true;
 	AllocatedResourceMap.Add(RHIBuffer, ResourceData);
 
 	if (!RHIBuffer->IsBarrierTrackingInitialized())
@@ -1615,32 +1610,26 @@ FRHITransientBuffer* FValidationTransientResourceAllocator::CreateBuffer(const F
 
 void FValidationTransientResourceAllocator::DeallocateMemory(FRHITransientTexture* InTransientTexture, uint32 InPassIndex)
 {
-	check(!bFrozen);
+	check(InTransientTexture);
 
 	RHIAllocator->DeallocateMemory(InTransientTexture, InPassIndex);
 
-	// Mark memory as freed
-	FAllocatedResourceData* ResourceData = AllocatedResourceMap.Find(InTransientTexture->GetRHI());
-	check(ResourceData);
-	ResourceData->bMemoryAllocated = false;
+	checkf(AllocatedResourceMap.Contains(InTransientTexture->GetRHI()), TEXT("DeallocateMemory called on texture %s, but it is not marked as allocated."), InTransientTexture->GetName());
+	AllocatedResourceMap.Remove(InTransientTexture->GetRHI());
 }
 
 void FValidationTransientResourceAllocator::DeallocateMemory(FRHITransientBuffer* InTransientBuffer, uint32 InPassIndex)
 {
-	check(!bFrozen);
+	check(InTransientBuffer);
 
 	RHIAllocator->DeallocateMemory(InTransientBuffer, InPassIndex);
 
-	// Mark memory as freed
-	FAllocatedResourceData* ResourceData = AllocatedResourceMap.Find(InTransientBuffer->GetRHI());
-	check(ResourceData);
-	ResourceData->bMemoryAllocated = false;
+	checkf(AllocatedResourceMap.Contains(InTransientBuffer->GetRHI()), TEXT("DeallocateMemory called on buffer %s, but it is not marked as allocated."), InTransientBuffer->GetName());
+	AllocatedResourceMap.Remove(InTransientBuffer->GetRHI());
 }
 
-void FValidationTransientResourceAllocator::Flush(FRHICommandListImmediate& RHICmdList)
+void FValidationTransientResourceAllocator::Flush(FRHICommandListImmediate& RHICmdList, FRHITransientAllocationStats* OutHeapStats)
 {
-	check(!bFrozen);
-
 	RHICmdList.EnqueueLambda([AllocatedResourcesToInit = MoveTemp(AllocatedResourcesToInit)](FRHICommandListImmediate& InRHICmdList)
 	{
 		// Tracking will be re-initialized, so we need to flush any remaining references.
@@ -1648,60 +1637,33 @@ void FValidationTransientResourceAllocator::Flush(FRHICommandListImmediate& RHIC
 		InitBarrierTracking(AllocatedResourcesToInit);
 	});
 
-	RHIAllocator->Flush(RHICmdList);
+	RHIAllocator->Flush(RHICmdList, OutHeapStats);
 }
 
-void FValidationTransientResourceAllocator::Freeze(FRHICommandListImmediate& RHICmdList, FRHITransientHeapStats& OutHeapStats)
+void FValidationTransientResourceAllocator::Release(FRHICommandListImmediate& RHICmdList)
 {
-	check(!bFrozen);
-
 	// Check all allocated resource data and make sure all memory is freed again
 	{
-		TArray<const FAllocatedResourceData*> AllocatedResources;
-
-		for (const auto& KeyValue : AllocatedResourceMap)
-		{
-			const FAllocatedResourceData& ResourceData = KeyValue.Value;
-			if (ResourceData.bMemoryAllocated)
-			{
-				AllocatedResources.Emplace(&ResourceData);
-			}
-		}
-
-		if (AllocatedResources.Num() > 0)
+		if (AllocatedResourceMap.Num() > 0)
 		{
 			FString ErrorMessage = FString::Printf(
 				TRANSIENT_RESOURCE_LOG_PREFIX_REASON("Open transient allocations")
-				TEXT("%d Transient Resource allocations still have memory allocated. Call 'DeallocateMemory' on all transient allocated resources before calling Freeze.\n\n")
+				TEXT("%d Transient Resource allocations still have memory allocated. Call 'DeallocateMemory' on all transient allocated resources prior to releasing the allocator.\n\n")
 				TEXT("Resources with Allocated Memory:\n"),
-				AllocatedResources.Num());
-			for (const FAllocatedResourceData* ResourceData : AllocatedResources)
+				AllocatedResourceMap.Num());
+
+			for (const auto KeyValue : AllocatedResourceMap)
 			{
-				ErrorMessage += FString::Printf(TEXT("         %s (%s)\n"), ResourceData->DebugName, ResourceData->ResourceType == FAllocatedResourceData::EType::Texture ? TEXT("Texture") : TEXT("Buffer"));
+				const FAllocatedResourceData& ResourceData = KeyValue.Value;
+
+				ErrorMessage += FString::Printf(TEXT("         %s (%s)\n"), ResourceData.DebugName, ResourceData.ResourceType == FAllocatedResourceData::EType::Texture ? TEXT("Texture") : TEXT("Buffer"));
 			}
 			ErrorMessage += FString::Printf(TRANSIENT_RESOURCE_LOG_SUFFIX);
 			FValidationRHI::ReportValidationFailure(*ErrorMessage);
 		}
 	}
 
-	RHICmdList.EnqueueLambda([AllocatedResourcesToInit = MoveTemp(AllocatedResourcesToInit)](FRHICommandListImmediate& InRHICmdList)
-	{
-		// Tracking will be re-initialized, so we need to flush any remaining references.
-		static_cast<FValidationContext&>(InRHICmdList.GetContext()).FlushValidationOps();
-		InitBarrierTracking(AllocatedResourcesToInit);
-	});
-
-	bFrozen = true;
-
-	RHIAllocator->Freeze(RHICmdList, OutHeapStats);
-}
-
-void FValidationTransientResourceAllocator::Release(FRHICommandListImmediate& RHICmdList)
-{
-	check(bFrozen);
-
 	RHIAllocator->Release(RHICmdList);
-	bReleased = true;
 	delete this;
 }
 
diff --git a/Engine/Source/Runtime/RHI/Public/RHITransientResourceAllocator.h b/Engine/Source/Runtime/RHI/Public/RHITransientResourceAllocator.h
--- a/Engine/Source/Runtime/RHI/Public/RHITransientResourceAllocator.h
+++ b/Engine/Source/Runtime/RHI/Public/RHITransientResourceAllocator.h
@@ -6,126 +6,265 @@
 #include "RHIResources.h"
 
 class FRHICommandListImmediate;
+class FRHITransientHeap;
+class FRHITransientPagePool;
 
-enum class ERHITransientResourceType
+struct FRHITransientPageSpan
+{
+	// Offset of the span in the page pool in pages. 
+	uint16 Offset = 0;
+
+	// Number of pages in the span.
+	uint16 Count = 0;
+};
+
+/** Represents an allocation from a transient page pool. */
+struct FRHITransientPagePoolAllocation
+{
+	bool IsValid() const { return Pool != nullptr; }
+
+	// The transient page pool which made the allocation.
+	FRHITransientPagePool* Pool = nullptr;
+
+	// A unique hash identifying this allocation to the allocator implementation.
+	uint64 Hash = 0;
+
+	// The index identifying the allocation to the page pool.
+	uint16 SpanIndex = 0;
+
+	// Offsets into the array of spans for the allocator implementation.
+	uint16 SpanOffsetMin = 0;
+	uint16 SpanOffsetMax = 0;
+};
+
+/** Represents a full set of page allocations from multiple page pools. */
+struct FRHITransientPageAllocation
+{
+	// The list of allocations by pool.
+	TArray<FRHITransientPagePoolAllocation> PoolAllocations;
+
+	// The full list of spans indexed by each allocation.
+	TArray<FRHITransientPageSpan> Spans;
+};
+
+/** Represents an allocation from the transient heap. */
+struct FRHITransientHeapAllocation
+{
+	bool IsValid() const { return Size != 0; }
+
+	// Transient heap which made the allocation.
+	FRHITransientHeap* Heap = nullptr;
+
+	// Size of the allocation made from the allocator (aligned).
+	uint64 Size = 0;
+
+	// Offset in the transient heap; front of the heap starts at 0.
+	uint64 Offset = 0;
+
+	// Number of bytes of padding were added to the offset.
+	uint32 AlignmentPad = 0;
+};
+
+enum class ERHITransientResourceType : uint8
 {
 	Texture,
 	Buffer
 };
 
-struct FRHITransientResourceStats
+enum class ERHITransientAllocationType : uint8
 {
-	uint64 HeapOffsetMin = 0;
-	uint64 HeapOffsetMax = 0;
-	uint32 HeapIndex = 0;
+	Heap,
+	Page
 };
 
 class FRHITransientResource
 {
 public:
-	FRHITransientResource(FRHIResource* InResource, uint64 InHash, ERHITransientResourceType InType)
+	static const uint32 kInvalidPassIndex = TNumericLimits<uint32>::Max();
+
+	FRHITransientResource(
+		FRHIResource* InResource,
+		uint64 InGpuVirtualAddress,
+		uint64 InHash,
+		uint64 InSize,
+		ERHITransientAllocationType InAllocationType,
+		ERHITransientResourceType InResourceType)
 		: Resource(InResource)
+		, GpuVirtualAddress(InGpuVirtualAddress)
 		, Hash(InHash)
-		, Type(InType)
+		, Size(InSize)
+		, AllocationType(InAllocationType)
+		, ResourceType(InResourceType)
 	{}
 
 	virtual ~FRHITransientResource() = default;
 
-	// (Internal) Initializes the transient resource with a new allocation / name.
-	virtual void Init(const TCHAR* InName, uint32 InAllocationIndex, uint32 InAcquirePassIndex, const FRHITransientResourceStats& InStats)
+	//////////////////////////////////////////////////////////////////////////////////////////////////
+	//! Internal Allocator API
+
+	virtual void Acquire(const TCHAR* InName, uint32 InAcquirePassIndex, uint64 InAllocatorCycle)
 	{
 		Name = InName;
-		AllocationIndex = InAllocationIndex;
 		AcquirePasses = TInterval<uint32>(0, InAcquirePassIndex);
-		DiscardPasses = TInterval<uint32>(0, TNumericLimits<uint32>::Max());
+		DiscardPasses = TInterval<uint32>(kInvalidPassIndex, kInvalidPassIndex);
+		AcquireCycle = InAllocatorCycle;
+		AcquireCount++;
 		AliasingOverlaps.Reset();
-		Stats = InStats;
 	}
 
-	// (Internal) Assigns the discard pass index.
-	FORCEINLINE void SetDiscardPass(uint32 PassIndex) { DiscardPasses.Min = PassIndex; }
+	void Discard(uint32 InDiscardPassIndex)
+	{
+		DiscardPasses.Min = InDiscardPassIndex;
+	}
 
-	// (Internal) Adds a new transient resource overlap.
-	FORCEINLINE void AddAliasingOverlap(FRHITransientResource* InResource)
+	void AddAliasingOverlap(FRHITransientResource* InResource)
 	{
 		AliasingOverlaps.Emplace(InResource->GetRHI(), InResource->IsTexture() ? FRHITransientAliasingOverlap::EType::Texture : FRHITransientAliasingOverlap::EType::Buffer);
 
+		check(InResource->DiscardPasses.Min != kInvalidPassIndex);
+
 		InResource->DiscardPasses.Max = FMath::Min(InResource->DiscardPasses.Max,             AcquirePasses.Max);
 		            AcquirePasses.Min = FMath::Max(            AcquirePasses.Min, InResource->DiscardPasses.Min);
 	}
 
+	FRHITransientHeapAllocation& GetHeapAllocation()
+	{
+		check(AllocationType == ERHITransientAllocationType::Heap);
+		return HeapAllocation;
+	}
+
+	const FRHITransientHeapAllocation& GetHeapAllocation() const
+	{
+		check(AllocationType == ERHITransientAllocationType::Heap);
+		return HeapAllocation;
+	}
+
+	FRHITransientPageAllocation& GetPageAllocation()
+	{
+		check(IsPageAllocated());
+		return PageAllocation;
+	}
+
+	const FRHITransientPageAllocation& GetPageAllocation() const
+	{
+		check(IsPageAllocated());
+		return PageAllocation;
+	}
+
+	//////////////////////////////////////////////////////////////////////////////////////////////////
+
 	// Returns the underlying RHI resource.
-	FORCEINLINE FRHIResource* GetRHI() const { return Resource; }
+	FRHIResource* GetRHI() const { return Resource; }
+
+	// Returns the gpu virtual address of the transient resource.
+	uint64 GetGpuVirtualAddress() const { return GpuVirtualAddress; }
 
 	// Returns the name assigned to the transient resource at allocation time.
-	FORCEINLINE const TCHAR* GetName() const { return Name; }
+	const TCHAR* GetName() const { return Name; }
+
+	// Returns the hash used to uniquely identify this resource if cached.
+	uint64 GetHash() const { return Hash; }
+
+	// Returns the required size in bytes of the resource.
+	uint64 GetSize() const { return Size; }
 
-	// (Internal)Returns the hash used to uniquely identify this resource if cached.
-	FORCEINLINE uint64 GetHash() const { return Hash; }
+	// Returns the last allocator cycle this resource was acquired.
+	uint64 GetAcquireCycle() const { return AcquireCycle; }
 
-	// (Internal) Returns the platform-specific allocation index.
-	FORCEINLINE uint32 GetAllocationIndex() const { return AllocationIndex; }
+	// Returns the number of times Acquire has been called.
+	uint32 GetAcquireCount() const { return AcquireCount; }
+
+	// Returns the list of aliasing overlaps used when transitioning the resource.
+	TConstArrayView<FRHITransientAliasingOverlap> GetAliasingOverlaps() const { return AliasingOverlaps; }
 
 	// Returns the pass index which may end acquiring this resource.
-	FORCEINLINE TInterval<uint32> GetAcquirePasses() const { return AcquirePasses; }
+	TInterval<uint32> GetAcquirePasses() const { return AcquirePasses; }
 
 	// Returns the pass index which discarded this resource.
-	FORCEINLINE TInterval<uint32> GetDiscardPasses() const { return DiscardPasses; }
+	TInterval<uint32> GetDiscardPasses() const { return DiscardPasses; }
+
+	// Returns whether this resource is still in an acquired state.
+	bool IsAcquired() const { return DiscardPasses.Min == kInvalidPassIndex; }
+
+	ERHITransientResourceType GetResourceType() const { return ResourceType; }
 
-	// Returns the aliasing overlaps for this resource.
-	FORCEINLINE TConstArrayView<FRHITransientAliasingOverlap> GetAliasingOverlaps() const { return AliasingOverlaps; }
+	bool IsTexture() const { return ResourceType == ERHITransientResourceType::Texture; }
+	bool IsBuffer()  const { return ResourceType == ERHITransientResourceType::Buffer; }
 
-	FORCEINLINE ERHITransientResourceType GetType() const { return Type; }
-	FORCEINLINE bool IsTexture() const { return Type == ERHITransientResourceType::Texture; }
-	FORCEINLINE bool IsBuffer() const { return Type == ERHITransientResourceType::Buffer; }
+	ERHITransientAllocationType GetAllocationType() const { return AllocationType; }
 
-	FORCEINLINE const FRHITransientResourceStats& GetStats() const { return Stats; }
+	bool IsHeapAllocated() const { return AllocationType == ERHITransientAllocationType::Heap; }
+	bool IsPageAllocated() const { return AllocationType == ERHITransientAllocationType::Page; }
 
 private:
 	// Underlying RHI resource.
 	TRefCountPtr<FRHIResource> Resource;
 
+	// The Gpu virtual address of the RHI resource.
+	uint64 GpuVirtualAddress = 0;
+
 	// The hash used to uniquely identify this resource if cached.
 	uint64 Hash;
 
+	// Size of the resource in bytes.
+	uint64 Size;
+
+	// Alignment of the resource in bytes.
+	uint32 Alignment;
+
+	// Tracks the number of times Acquire has been called.
+	uint32 AcquireCount = 0;
+
+	// Cycle count used to deduce age of the resource.
+	uint64 AcquireCycle = 0;
+
 	// Debug name of the resource. Updated with each allocation.
-	const TCHAR* Name{};
+	const TCHAR* Name = nullptr;
+
+	FRHITransientHeapAllocation HeapAllocation;
+	FRHITransientPageAllocation PageAllocation;
 
 	// List of aliasing resources overlapping with this one.
 	TArray<FRHITransientAliasingOverlap> AliasingOverlaps;
 
-	// Type of the underlying RHI resource.
-	ERHITransientResourceType Type;
-
-	// An index to the underlying allocation info on the internal platform allocator.
-	uint32 AllocationIndex = ~0u;
+	// Start -> End split pass index intervals for acquire / discard operations.
 	TInterval<uint32> AcquirePasses = TInterval<uint32>(0, 0);
 	TInterval<uint32> DiscardPasses = TInterval<uint32>(0, 0);
 
-	FRHITransientResourceStats Stats;
+	ERHITransientAllocationType AllocationType;
+	ERHITransientResourceType ResourceType;
 };
 
 class RHI_API FRHITransientTexture final : public FRHITransientResource
 {
 public:
-	FRHITransientTexture(FRHITexture* InTexture, uint64 InHash, const FRHITextureCreateInfo& InCreateInfo)
-		: FRHITransientResource(InTexture, InHash, ERHITransientResourceType::Texture)
+	FRHITransientTexture(
+		FRHITexture* InTexture,
+		uint64 InGpuVirtualAddress,
+		uint64 InHash,
+		uint64 InSize,
+		ERHITransientAllocationType InAllocationType,
+		const FRHITextureCreateInfo& InCreateInfo)
+		: FRHITransientResource(InTexture, InGpuVirtualAddress, InHash, InSize, InAllocationType, ERHITransientResourceType::Texture)
 		, CreateInfo(InCreateInfo)
 	{}
 
-	void Init(const TCHAR* InName, uint32 InAllocationIndex, uint32 InPassIndex, const FRHITransientResourceStats& InStats) override;
+	//////////////////////////////////////////////////////////////////////////////////////////////////
+	//! Internal Allocator API
+	void Acquire(const TCHAR* InName, uint32 InAcquirePassIndex, uint64 InInitCycle) override;
+	//////////////////////////////////////////////////////////////////////////////////////////////////
 
 	// Returns the underlying RHI texture.
 	FRHITexture* GetRHI() const { return static_cast<FRHITexture*>(FRHITransientResource::GetRHI()); }
 
 	// Returns the create info struct used when creating this texture.
-	FORCEINLINE const FRHITextureCreateInfo& GetCreateInfo() const { return CreateInfo; }
+	const FRHITextureCreateInfo& GetCreateInfo() const { return CreateInfo; }
 
 	// Finds a UAV matching the descriptor in the cache or creates a new one and updates the cache.
-	FORCEINLINE FRHIUnorderedAccessView* GetOrCreateUAV(const FRHITextureUAVCreateInfo& InCreateInfo) { return ViewCache.GetOrCreateUAV(GetRHI(), InCreateInfo); }
+	FRHIUnorderedAccessView* GetOrCreateUAV(const FRHITextureUAVCreateInfo& InCreateInfo) { return ViewCache.GetOrCreateUAV(GetRHI(), InCreateInfo); }
 
 	// Finds a SRV matching the descriptor in the cache or creates a new one and updates the cache.
-	FORCEINLINE FRHIShaderResourceView* GetOrCreateSRV(const FRHITextureSRVCreateInfo& InCreateInfo) { return ViewCache.GetOrCreateSRV(GetRHI(), InCreateInfo); }
+	FRHIShaderResourceView* GetOrCreateSRV(const FRHITextureSRVCreateInfo& InCreateInfo) { return ViewCache.GetOrCreateSRV(GetRHI(), InCreateInfo); }
 
 	// The create info describing the texture.
 	const FRHITextureCreateInfo CreateInfo;
@@ -137,24 +276,33 @@ public:
 class RHI_API FRHITransientBuffer final : public FRHITransientResource
 {
 public:
-	FORCEINLINE FRHITransientBuffer(FRHIBuffer* InBuffer, uint64 InHash, const FRHIBufferCreateInfo& InCreateInfo)
-		: FRHITransientResource(InBuffer, InHash, ERHITransientResourceType::Buffer)
+	FRHITransientBuffer(
+		FRHIBuffer* InBuffer,
+		uint64 InGpuVirtualAddress,
+		uint64 InHash,
+		uint64 InSize,
+		ERHITransientAllocationType InAllocationType,
+		const FRHIBufferCreateInfo& InCreateInfo)
+		: FRHITransientResource(InBuffer, InGpuVirtualAddress, InHash, InSize, InAllocationType, ERHITransientResourceType::Buffer)
 		, CreateInfo(InCreateInfo)
 	{}
 
-	void Init(const TCHAR* InName, uint32 InAllocationIndex, uint32 InPassIndex, const FRHITransientResourceStats& InStats) override;
+	//////////////////////////////////////////////////////////////////////////////////////////////////
+	//! Internal Allocator API
+	void Acquire(const TCHAR* InName, uint32 InAcquirePassIndex, uint64 InInitCycle) override;
+	//////////////////////////////////////////////////////////////////////////////////////////////////
 
 	// Returns the underlying RHI buffer.
-	FORCEINLINE FRHIBuffer* GetRHI() const { return static_cast<FRHIBuffer*>(FRHITransientResource::GetRHI()); }
+	FRHIBuffer* GetRHI() const { return static_cast<FRHIBuffer*>(FRHITransientResource::GetRHI()); }
 
 	// Returns the create info used when creating this buffer.
-	FORCEINLINE const FRHIBufferCreateInfo& GetCreateInfo() const { return CreateInfo; }
+	const FRHIBufferCreateInfo& GetCreateInfo() const { return CreateInfo; }
 
 	// Finds a UAV matching the descriptor in the cache or creates a new one and updates the cache.
-	FORCEINLINE FRHIUnorderedAccessView* GetOrCreateUAV(const FRHIBufferUAVCreateInfo& InCreateInfo) { return ViewCache.GetOrCreateUAV(GetRHI(), InCreateInfo); }
+	FRHIUnorderedAccessView* GetOrCreateUAV(const FRHIBufferUAVCreateInfo& InCreateInfo) { return ViewCache.GetOrCreateUAV(GetRHI(), InCreateInfo); }
 
 	// Finds a SRV matching the descriptor in the cache or creates a new one and updates the cache.
-	FORCEINLINE FRHIShaderResourceView* GetOrCreateSRV(const FRHIBufferSRVCreateInfo& InCreateInfo) { return ViewCache.GetOrCreateSRV(GetRHI(), InCreateInfo); }
+	FRHIShaderResourceView* GetOrCreateSRV(const FRHIBufferSRVCreateInfo& InCreateInfo) { return ViewCache.GetOrCreateSRV(GetRHI(), InCreateInfo); }
 
 	// The create info describing the texture.
 	const FRHIBufferCreateInfo CreateInfo;
@@ -163,18 +311,29 @@ public:
 	FRHIBufferViewCache ViewCache;
 };
 
-struct FRHITransientHeapStats
+class FRHITransientAllocationStats
 {
-	struct FHeap
+public:
+	struct FAllocation
 	{
-		// Maximum number of bytes used in the heap.
-		uint64 WatermarkSize{};
+		uint64 OffsetMin = 0;
+		uint64 OffsetMax = 0;
+		uint32 MemoryRangeIndex = 0;
+	};
+
+	using FAllocationArray = TArray<FAllocation, TInlineAllocator<2>>;
+
+	struct FMemoryRange
+	{
+		// Number of bytes available for use in the memory range.
+		uint64 Capacity = 0;
 
-		// Number of bytes allocated for use in the heap.
-		uint64 Capacity{};
+		// Number of bytes allocated for use in the memory range.
+		uint64 CommitSize = 0;
 	};
 
-	TArray<FHeap, TInlineAllocator<4>> Heaps;
+	TArray<FMemoryRange> MemoryRanges;
+	TMap<const FRHITransientResource*, FAllocationArray> Resources;
 };
 
 class IRHITransientResourceAllocator
@@ -190,15 +349,9 @@ public:
 	virtual void DeallocateMemory(FRHITransientTexture* InTexture, uint32 InPassIndex) = 0;
 	virtual void DeallocateMemory(FRHITransientBuffer* InBuffer, uint32 InPassIndex) = 0;
 
-	// Flushes any pending allocations in preparation for rendering. Resources are not required to be deallocated yet. Use when interleaving execution with allocation.
-	virtual void Flush(FRHICommandListImmediate& RHICmdList) {}
+	// Flushes any pending allocations prior to rendering. Optionally emits stats if OutStats is valid.
+	virtual void Flush(FRHICommandListImmediate& RHICmdList, FRHITransientAllocationStats* OutStats = nullptr) = 0;
 
-	// Freezes all allocations and validates that all resources have their memory deallocated.
-	virtual void Freeze(FRHICommandListImmediate& RHICmdList, FRHITransientHeapStats& OutHeapStats) = 0;
-
-	// Releases the transient allocator and deletes the instance. Any FRHITransientResource* access after this call is not allowed.
-	virtual void Release(FRHICommandListImmediate& RHICmdList)
-	{
-		delete this;
-	}
+	// Releases this instance of the transient allocator. Invalidates any outstanding transient resources.
+	virtual void Release(FRHICommandListImmediate& RHICmdList) { delete this; }
 };
\ No newline at end of file
diff --git a/Engine/Source/Runtime/RHI/Public/RHIValidationTransientResourceAllocator.h b/Engine/Source/Runtime/RHI/Public/RHIValidationTransientResourceAllocator.h
--- a/Engine/Source/Runtime/RHI/Public/RHIValidationTransientResourceAllocator.h
+++ b/Engine/Source/Runtime/RHI/Public/RHIValidationTransientResourceAllocator.h
@@ -20,17 +20,13 @@ public:
 	virtual FRHITransientBuffer* CreateBuffer(const FRHIBufferCreateInfo& InCreateInfo, const TCHAR* InDebugName, uint32 InPassIndex) override final;
 	virtual void DeallocateMemory(FRHITransientTexture* InTexture, uint32 InPassIndex) override final;
 	virtual void DeallocateMemory(FRHITransientBuffer* InBuffer, uint32 InPassIndex) override final;
-	virtual void Flush(FRHICommandListImmediate&) override final;
-	virtual void Freeze(FRHICommandListImmediate&, FRHITransientHeapStats&) override final;
+	virtual void Flush(FRHICommandListImmediate&, FRHITransientAllocationStats*) override final;
 	virtual void Release(FRHICommandListImmediate&) override final;
 
 private:
 	// Actual RHI transient allocator which will get all functions forwarded
 	IRHITransientResourceAllocator* RHIAllocator = nullptr;
 
-	bool bFrozen = false;
-	bool bReleased = false;
-
 	// All the allocated resources on the transient allocator
 	struct FAllocatedResourceData
 	{
@@ -42,7 +38,6 @@ private:
 
 		const TCHAR* DebugName = nullptr;
 		EType ResourceType = EType::Texture;
-		bool bMemoryAllocated = false;
 
 		struct FTexture
 		{
diff --git a/Engine/Source/Runtime/RHICore/Private/RHICoreTransientResourceAllocator.cpp b/Engine/Source/Runtime/RHICore/Private/RHICoreTransientResourceAllocator.cpp
--- a/Engine/Source/Runtime/RHICore/Private/RHICoreTransientResourceAllocator.cpp
+++ b/Engine/Source/Runtime/RHICore/Private/RHICoreTransientResourceAllocator.cpp
@@ -3,9 +3,11 @@
 #include "RHICoreTransientResourceAllocator.h"
 #include "HAL/LowLevelMemTracker.h"
 #include "HAL/LowLevelMemStats.h"
+#include "ProfilingDebugging/CountersTrace.h"
+#include "RHICore.h"
 
 static int32 GRHITransientAllocatorMinimumHeapSize = 128;
-static FAutoConsoleVariableRef CVarGRHITransientAllocatorMinimumHeapSize(
+static FAutoConsoleVariableRef CVarRHITransientAllocatorMinimumHeapSize(
 	TEXT("RHI.TransientAllocator.MinimumHeapSize"),
 	GRHITransientAllocatorMinimumHeapSize,
 	TEXT("Minimum size of an RHI transient heap in MB. Heaps will default to this size and grow to the maximum based on the first allocation (Default 128)."),
@@ -32,207 +34,245 @@ static FAutoConsoleVariableRef CVarRHITransientAllocatorTextureCacheSize(
 	TEXT("The maximum number of RHI textures to cache on each heap before garbage collecting."),
 	ECVF_ReadOnly);
 
-static int32 GRHITransientAllocatorHeapGarbageCollectLatency = 20;
-static FAutoConsoleVariableRef CVarRHITransientAllocatorHeapGarbageCollectLatency(
-	TEXT("RHI.TransientAllocator.HeapGarbageCollectLatency"),
-	GRHITransientAllocatorHeapGarbageCollectLatency,
-	TEXT("Amount of update cycles before a heap is deleted when not used."),
+static int32 GRHITransientAllocatorGarbageCollectLatency = 16;
+static FAutoConsoleVariableRef CVarRHITransientAllocatorGarbageCollectLatency(
+	TEXT("RHI.TransientAllocator.GarbageCollectLatency"),
+	GRHITransientAllocatorGarbageCollectLatency,
+	TEXT("Amount of update cycles before memory is reclaimed."),
 	ECVF_ReadOnly);
 
+#if RHICORE_TRANSIENT_ALLOCATOR_DEBUG
+static int32 GRHITransientAllocatorLogCycles = 0;
+static FAutoConsoleVariableRef CVarRHITransientAllocatorLogCycles(
+	TEXT("RHI.TransientAllocator.LogCycles"),
+	GRHITransientAllocatorLogCycles,
+	TEXT("Logs resources if they have been allocated for more than X allocator Flush() cycles; use to identify and log persistent allocations."),
+	ECVF_RenderThreadSafe);
+#endif
+
+TRACE_DECLARE_INT_COUNTER(TransientResourceCreateCount, TEXT("TransientAllocator/ResourceCreateCount"));
+
+TRACE_DECLARE_INT_COUNTER(TransientTextureCreateCount, TEXT("TransientAllocator/TextureCreateCount"));
+TRACE_DECLARE_INT_COUNTER(TransientTextureCount, TEXT("TransientAllocator/TextureCount"));
+TRACE_DECLARE_INT_COUNTER(TransientTextureCacheSize, TEXT("TransientAllocator/TextureCacheSize"));
+TRACE_DECLARE_FLOAT_COUNTER(TransientTextureCacheHitPercentage, TEXT("TransientAllocator/TextureCacheHitPercentage"));
+
+TRACE_DECLARE_INT_COUNTER(TransientBufferCreateCount, TEXT("TransientAllocator/BufferCreateCount"));
+TRACE_DECLARE_INT_COUNTER(TransientBufferCount, TEXT("TransientAllocator/BufferCount"));
+TRACE_DECLARE_INT_COUNTER(TransientBufferCacheSize, TEXT("TransientAllocator/BufferCacheSize"));
+TRACE_DECLARE_FLOAT_COUNTER(TransientBufferCacheHitPercentage, TEXT("TransientAllocator/BufferCacheHitPercentage"));
+
+TRACE_DECLARE_INT_COUNTER(TransientPageMapCount, TEXT("TransientAllocator/PageMapCount"));
+TRACE_DECLARE_INT_COUNTER(TransientPageAllocateCount, TEXT("TransientAllocator/PageAllocateCount"));
+TRACE_DECLARE_INT_COUNTER(TransientPageSpanCount, TEXT("TransientAllocator/PageSpanCount"));
+
+TRACE_DECLARE_INT_COUNTER(TransientMemoryRangeCount, TEXT("TransientAllocator/MemoryRangeCount"));
+
+TRACE_DECLARE_MEMORY_COUNTER(TransientMemoryUsed, TEXT("TransientAllocator/MemoryUsed"));
+TRACE_DECLARE_MEMORY_COUNTER(TransientMemoryRequested, TEXT("TransientAllocator/MemoryRequested"));
+
 DECLARE_STATS_GROUP(TEXT("RHI: Transient Memory"), STATGROUP_RHITransientMemory, STATCAT_Advanced);
 
-DECLARE_MEMORY_STAT(TEXT("Memory Allocated"), STAT_RHITransientMemoryAllocated, STATGROUP_RHITransientMemory);
-DECLARE_MEMORY_STAT(TEXT("Memory Requested"), STAT_RHITransientMemoryRequested, STATGROUP_RHITransientMemory);
 DECLARE_MEMORY_STAT(TEXT("Memory Used"), STAT_RHITransientMemoryUsed, STATGROUP_RHITransientMemory);
-DECLARE_MEMORY_STAT(TEXT("Buffer Memory Used"), STAT_RHITransientBufferMemoryUsed, STATGROUP_RHITransientMemory);
+DECLARE_MEMORY_STAT(TEXT("Memory Aliased"), STAT_RHITransientMemoryAliased, STATGROUP_RHITransientMemory);
+DECLARE_MEMORY_STAT(TEXT("Memory Requested"), STAT_RHITransientMemoryRequested, STATGROUP_RHITransientMemory);
 DECLARE_MEMORY_STAT(TEXT("Buffer Memory Requested"), STAT_RHITransientBufferMemoryRequested, STATGROUP_RHITransientMemory);
-DECLARE_MEMORY_STAT(TEXT("Texture Memory Used"), STAT_RHITransientTextureMemoryUsed, STATGROUP_RHITransientMemory);
 DECLARE_MEMORY_STAT(TEXT("Texture Memory Requested"), STAT_RHITransientTextureMemoryRequested, STATGROUP_RHITransientMemory);
 
 DECLARE_DWORD_ACCUMULATOR_STAT(TEXT("Resources"), STAT_RHITransientResources, STATGROUP_RHITransientMemory);
 DECLARE_DWORD_ACCUMULATOR_STAT(TEXT("Textures"), STAT_RHITransientTextures, STATGROUP_RHITransientMemory);
 DECLARE_DWORD_ACCUMULATOR_STAT(TEXT("Buffers"), STAT_RHITransientBuffers, STATGROUP_RHITransientMemory);
-DECLARE_DWORD_ACCUMULATOR_STAT(TEXT("Heaps"), STAT_RHITransientHeaps, STATGROUP_RHITransientMemory);
 
 DECLARE_LLM_MEMORY_STAT(TEXT("RHI Transient Resources"), STAT_RHITransientResourcesLLM, STATGROUP_LLMFULL);
 LLM_DEFINE_TAG(RHITransientResources, NAME_None, NAME_None, GET_STATFNAME(STAT_RHITransientResourcesLLM), GET_STATFNAME(STAT_EngineSummaryLLM));
 
 //////////////////////////////////////////////////////////////////////////
 
-template <typename CreateFunctionType>
-FRHITransientTexture* FRHITransientHeap::AcquireTexture(const FRHITextureCreateInfo& CreateInfo, uint64 HeapOffset, CreateFunctionType CreateFunction)
-{
-	const uint64 Hash = ComputeHash(CreateInfo, HeapOffset);
-	FRHITransientTexture* Texture = Textures.Acquire(Hash, CreateFunction);
-	AllocatedTextures.Emplace(Texture);
-	return Texture;
-}
-
-template <typename CreateFunctionType>
-FRHITransientBuffer* FRHITransientHeap::AcquireBuffer(const FRHIBufferCreateInfo& CreateInfo, uint64 HeapOffset, CreateFunctionType CreateFunction)
-{
-	const uint64 Hash = ComputeHash(CreateInfo, HeapOffset);
-	FRHITransientBuffer* Buffer = Buffers.Acquire(Hash, CreateFunction);
-	AllocatedBuffers.Emplace(Buffer);
-	return Buffer;
-}
-
-void FRHITransientHeap::ForfeitResources()
+void FRHITransientMemoryStats::Submit(uint64 UsedSize)
 {
-	Textures.Forfeit(AllocatedTextures);
-	AllocatedTextures.Reset();
-
-	Buffers.Forfeit(AllocatedBuffers);
-	AllocatedBuffers.Reset();
+	const int32 CreateResourceCount = Textures.CreateCount + Buffers.CreateCount;
+	const int64 MemoryUsed = UsedSize;
+	const int64 MemoryRequested = AliasedSize;
+	const float ToMB = 1.0f / (1024.0f * 1024.0f);
+
+	TRACE_COUNTER_SET(TransientResourceCreateCount, CreateResourceCount);
+	TRACE_COUNTER_SET(TransientTextureCreateCount, Textures.CreateCount);
+	TRACE_COUNTER_SET(TransientBufferCreateCount, Buffers.CreateCount);
+	TRACE_COUNTER_SET(TransientMemoryUsed, MemoryUsed);
+	TRACE_COUNTER_SET(TransientMemoryRequested, MemoryRequested);
+
+	CSV_CUSTOM_STAT_GLOBAL(TransientResourceCreateCount, CreateResourceCount, ECsvCustomStatOp::Set);
+	CSV_CUSTOM_STAT_GLOBAL(TransientMemoryUsedMB, static_cast<float>(MemoryUsed * ToMB) , ECsvCustomStatOp::Set);
+	CSV_CUSTOM_STAT_GLOBAL(TransientMemoryAliasedMB, static_cast<float>(MemoryRequested * ToMB), ECsvCustomStatOp::Set);
+
+	SET_MEMORY_STAT(STAT_RHITransientMemoryUsed, UsedSize);
+	SET_MEMORY_STAT(STAT_RHITransientMemoryAliased, AliasedSize);
+	SET_MEMORY_STAT(STAT_RHITransientMemoryRequested, Textures.AllocatedSize + Buffers.AllocatedSize);
+	SET_MEMORY_STAT(STAT_RHITransientBufferMemoryRequested, Buffers.AllocatedSize);
+	SET_MEMORY_STAT(STAT_RHITransientTextureMemoryRequested, Textures.AllocatedSize);
+
+	SET_DWORD_STAT(STAT_RHITransientTextures, Textures.AllocationCount);
+	SET_DWORD_STAT(STAT_RHITransientBuffers, Buffers.AllocationCount);
+	SET_DWORD_STAT(STAT_RHITransientResources, Textures.AllocationCount + Buffers.AllocationCount);
+
+	Reset();
 }
 
 //////////////////////////////////////////////////////////////////////////
 
-FRHITransientResourceSystemInitializer FRHITransientResourceSystemInitializer::CreateDefault()
+void FRHITransientResourceOverlapTracker::Track(FRHITransientResource* TransientResource, uint32 PageOffsetMin, uint32 PageOffsetMax)
 {
-	FRHITransientResourceSystemInitializer Initializer;
-	Initializer.MinimumHeapSize = GRHITransientAllocatorMinimumHeapSize * 1024 * 1024;
-	Initializer.MaximumHeapSize = GRHITransientAllocatorMaximumHeapSize * 1024 * 1024;
-	Initializer.BufferCacheSize = GRHITransientAllocatorBufferCacheSize;
-	Initializer.TextureCacheSize = GRHITransientAllocatorTextureCacheSize;
-	Initializer.GarbageCollectLatency = GRHITransientAllocatorHeapGarbageCollectLatency;
-	return Initializer;
-}
+	check(TransientResource);
 
-FRHITransientResourceSystem::~FRHITransientResourceSystem()
-{
-	ReleaseHeaps();
-}
+	FResourceRange ResourceRangeNew;
+	ResourceRangeNew.ResourceIndex = Resources.Emplace(TransientResource);
+	ResourceRangeNew.PageOffsetMin = PageOffsetMin;
+	ResourceRangeNew.PageOffsetMax = PageOffsetMax;
 
-FRHITransientHeap* FRHITransientResourceSystem::AcquireHeap(uint64 FirstAllocationSize, uint32 FirstAllocationAlignment, ERHITransientHeapFlags FirstAllocationHeapFlags)
-{
+	int32 InsertIndex = Algo::LowerBound(ResourceRanges, ResourceRangeNew, [](const FResourceRange& Lhs, const FResourceRange& Rhs)
+	{
+		return Lhs.PageOffsetMax <= Rhs.PageOffsetMin;
+	});
+
+	for (int32 Index = InsertIndex; Index < ResourceRanges.Num(); ++Index)
 	{
-		FScopeLock Lock(&HeapCriticalSection);
+		FResourceRange& ResourceRangeOld = ResourceRanges[Index];
 
-		for (int32 HeapIndex = 0; HeapIndex < Heaps.Num(); ++HeapIndex)
+		// If the old range starts later in the heap and doesn't overlap, the sort invariant guarantees no future range will overlap.
+		if (ResourceRangeOld.PageOffsetMin >= ResourceRangeNew.PageOffsetMax)
 		{
-			FRHITransientHeap* Heap = Heaps[HeapIndex];
-
-			if (Heap->IsAllocationSupported(FirstAllocationSize, FirstAllocationAlignment, FirstAllocationHeapFlags))
-			{
-				Heaps.RemoveAt(HeapIndex);
-				return Heap;
-			}
+			break;
 		}
-	}
 
-	FRHITransientHeapInitializer HeapInitializer;
-	HeapInitializer.Size = GetHeapSize(FirstAllocationSize);
-	HeapInitializer.Alignment = GetHeapAlignment(FirstAllocationAlignment);
-	HeapInitializer.Flags = (Initializer.bSupportsAllHeapFlags ? ERHITransientHeapFlags::AllowAll : FirstAllocationHeapFlags);
-	HeapInitializer.TextureCacheSize = Initializer.TextureCacheSize;
-	HeapInitializer.BufferCacheSize = Initializer.BufferCacheSize;
+		TransientResource->AddAliasingOverlap(Resources[ResourceRangeOld.ResourceIndex]);
 
-	LLM_SCOPE_BYTAG(RHITransientResources);
-	return CreateHeap(HeapInitializer);
-}
+		// Complete overlap.
+		if (ResourceRangeOld.PageOffsetMin >= ResourceRangeNew.PageOffsetMin && ResourceRangeOld.PageOffsetMax <= ResourceRangeNew.PageOffsetMax)
+		{
+			ResourceRanges.RemoveAt(Index, 1, false);
+			Index--;
+		}
+		// Partial overlap, can manifest as three cases:
+		else
+		{
+			bool bResizedOld = false;
+			FResourceRange ResourceRangeOldCopy = ResourceRangeOld;
 
-void FRHITransientResourceSystem::ForfeitHeaps(TConstArrayView<FRHITransientHeapState> InForfeitedHeaps)
-{
-	FScopeLock Lock(&HeapCriticalSection);
+			// 1) New:    ********
+			//            |||        ->
+			//    Old: ======             ===********
+			if (ResourceRangeOld.PageOffsetMin < ResourceRangeNew.PageOffsetMin)
+			{
+				ResourceRangeOld.PageOffsetMax = ResourceRangeNew.PageOffsetMin;
+				bResizedOld = true;
 
-	Heaps.Reserve(InForfeitedHeaps.Num());
-	for (FRHITransientHeapState const& State : InForfeitedHeaps)
-	{
-		State.Heap->ForfeitResources();
-		State.Heap->LastUsedGarbageCollectCycle = GarbageCollectCycle;
+				// This case should only ever happen on the first iteration.
+				check(Index == InsertIndex);
+				InsertIndex++;
+			}
+			else
+			{
+				// 2) New:    ********
+				//                |||      ->
+				//    Old:        ======         ********===
+				if (!bResizedOld)
+				{
+					ResourceRangeOld.PageOffsetMin = ResourceRangeNew.PageOffsetMax;
+				}
 
-		Heaps.Add(State.Heap);
-	}
+				// 3) New:    ********
+				//            ||||||||      ->
+				//    Old: ==============        ===********===
+				else
+				{
+					// Lower bound has been resized already; add an upper bound.
+					FResourceRange ResourceRangeOldUpper = ResourceRangeOldCopy;
+					Index = ResourceRanges.Insert(ResourceRangeOldUpper, Index + 1);
+				}
 
-	Algo::Sort(Heaps, [](const FRHITransientHeap* LHS, const FRHITransientHeap* RHS)
-	{
-		// Sort by smaller heap first.
-		if (LHS->GetCapacity() != RHS->GetCapacity())
-		{
-			return LHS->GetCapacity() < RHS->GetCapacity();
+				break;
+			}
 		}
+	}
 
-		// Sort next by most recently used first.
-		return LHS->GetLastUsedGarbageCollectCycle() >= RHS->GetLastUsedGarbageCollectCycle();
-	});
+	ResourceRanges.Insert(ResourceRangeNew, InsertIndex);
 }
 
-void FRHITransientResourceSystem::GarbageCollect()
-{
-	FScopeLock Lock(&HeapCriticalSection);
+//////////////////////////////////////////////////////////////////////////
 
-	for (int32 HeapIndex = 0; HeapIndex < Heaps.Num(); ++HeapIndex)
+template <typename ResourceContainerType>
+void LogActiveResources(const ResourceContainerType& Resources, uint64 CurrentCycle)
+{
+#if RHICORE_TRANSIENT_ALLOCATOR_DEBUG
+	if (!Resources.IsEmpty())
 	{
-		FRHITransientHeap* Heap = Heaps[HeapIndex];
+		FString ErrorMsg = FString::Printf(
+			TEXT("---------------------------------------------------------------------------------------------------------\n")
+			TEXT("-- %d Transient resources are still being held outside the allocator. Memory has been leaked."),
+			Resources.Num());
 
-		if (Heap->GetLastUsedGarbageCollectCycle() + Initializer.GarbageCollectLatency <= GarbageCollectCycle)
+		for (FRHITransientResource* Resource : Resources)
 		{
-			Heaps.RemoveAt(HeapIndex);
-			HeapIndex--;
-
-			delete Heap;
+			ErrorMsg += FString::Printf(TEXT("\tName: %s, Age: %llu cycles."), Resource->GetName(), CurrentCycle - Resource->GetAcquireCycle());
 		}
-	}
 
-	GarbageCollectCycle++;
+		UE_LOG(LogRHICore, Fatal,
+			TEXT("%s\n")
+			TEXT("---------------------------------------------------------------------------------------------------------\n"),
+			*ErrorMsg);
+	}
+#endif
 }
 
-void FRHITransientResourceSystem::ReleaseHeaps()
+template <typename ResourceContainerType>
+void LogLongLivedResources(const ResourceContainerType& Resources, uint64 CurrentCycle)
 {
-	FScopeLock Lock(&HeapCriticalSection);
-	for (FRHITransientHeap* Heap : Heaps)
+#if RHICORE_TRANSIENT_ALLOCATOR_DEBUG
+	if (GRHITransientAllocatorLogCycles <= 0 || Resources.IsEmpty())
 	{
-		delete Heap;
+		return;
 	}
-	Heaps.Empty();
-}
-
-void FRHITransientResourceSystem::UpdateStats()
-{
-	FStats Stats;
-	Stats.Textures = TextureStats;
-	Stats.Buffers = BufferStats;
 
-	uint32 HeapCount;
+	TArray<FRHITransientResource*> LongLivedResources;
 
+	for (FRHITransientResource* Resource : Resources)
 	{
-		FScopeLock Lock(&HeapCriticalSection);
-		for (FRHITransientHeap* Heap : Heaps)
+		if (CurrentCycle - Resource->GetAcquireCycle() > GRHITransientAllocatorLogCycles)
 		{
-			Stats.TotalMemoryUsed += Heap->GetCapacity();
+			LongLivedResources.Emplace(Resource);
 		}
-		HeapCount = Heaps.Num();
 	}
 
-	ReportStats(Stats);
-
-	SET_MEMORY_STAT(STAT_RHITransientMemoryAllocated, Stats.TotalMemoryUsed);
-	SET_MEMORY_STAT(STAT_RHITransientMemoryRequested, Stats.Textures.TotalSize + Stats.Buffers.TotalSize);
-	SET_MEMORY_STAT(STAT_RHITransientMemoryUsed, Stats.Textures.TotalSizeWithAliasing + Stats.Buffers.TotalSizeWithAliasing);
-	SET_MEMORY_STAT(STAT_RHITransientBufferMemoryRequested, Stats.Buffers.TotalSize);
-	SET_MEMORY_STAT(STAT_RHITransientBufferMemoryUsed, Stats.Buffers.TotalSizeWithAliasing);
-	SET_MEMORY_STAT(STAT_RHITransientTextureMemoryRequested, Stats.Textures.TotalSize);
-	SET_MEMORY_STAT(STAT_RHITransientTextureMemoryUsed, Stats.Textures.TotalSizeWithAliasing);
-
-	SET_DWORD_STAT(STAT_RHITransientTextures, Stats.Textures.AllocationCount);
-	SET_DWORD_STAT(STAT_RHITransientBuffers, Stats.Buffers.AllocationCount);
-	SET_DWORD_STAT(STAT_RHITransientResources, Stats.Textures.AllocationCount + Stats.Buffers.AllocationCount);
-	SET_DWORD_STAT(STAT_RHITransientHeaps, HeapCount);
+	if (!LongLivedResources.IsEmpty())
+	{
+		UE_LOG(LogRHICore, Warning,
+			TEXT("---------------------------------------------------------------------------------------------------------\n")
+			TEXT("-- %d transient %ss have been allocated for more than %d cycles. Please mark them non-transient to reduce fragmentation."),
+			LongLivedResources.Num(),
+			LongLivedResources[0]->GetResourceType() == ERHITransientResourceType::Texture ? TEXT("texture") : TEXT("buffer"),
+			GRHITransientAllocatorLogCycles);
+
+		for (FRHITransientResource* Resource : LongLivedResources)
+		{
+			UE_LOG(LogRHICore, Warning, TEXT("\tName: %s, Age: %llu cycles."), Resource->GetName(), CurrentCycle - Resource->GetAcquireCycle());
+		}
 
-	TextureStats = {};
-	BufferStats = {};
-}
+		UE_LOG(LogRHICore, Warning,
+			TEXT("---------------------------------------------------------------------------------------------------------\n"));
+	}
+#endif
+};
 
+//////////////////////////////////////////////////////////////////////////
+// Transient Resource Heap Allocator
 //////////////////////////////////////////////////////////////////////////
 
-FRHITransientHeapAllocator::FRHITransientHeapAllocator(const FRHITransientHeapInitializer& InInitializer, uint32 InHeapIndex, uint64 InHeapBaseGPUAddress)
-	: Initializer(InInitializer)
-	, HeapIndex(InHeapIndex)
-	, HeapBaseGPUAddress(InHeapBaseGPUAddress)
+FRHITransientHeapAllocator::FRHITransientHeapAllocator(uint64 InCapacity, uint32 InAlignment)
+	: Capacity(InCapacity)
+	, AlignmentMin(InAlignment)
 {
 	HeadHandle = CreateRange();
-	InsertRange(HeadHandle, 0, Initializer.Size);
+	InsertRange(HeadHandle, 0, Capacity);
 }
 
 FRHITransientHeapAllocator::~FRHITransientHeapAllocator()
@@ -255,7 +295,10 @@ FRHITransientHeapAllocation FRHITransientHeapAllocator::Allocate(uint64 Size, ui
 {
 	check(Size > 0);
 
-	Alignment = FMath::Max(Alignment, Initializer.Alignment);
+	if (Alignment < AlignmentMin)
+	{
+		Alignment = AlignmentMin;
+	}
 
 	FFindResult FindResult = FindFreeRange(Size, Alignment);
 
@@ -266,16 +309,16 @@ FRHITransientHeapAllocation FRHITransientHeapAllocator::Allocate(uint64 Size, ui
 
 	FRange& FoundRange = Ranges[FindResult.FoundHandle];
 
-	const uint64 AlignedSize   = FoundRange.Size   - FindResult.LeftoverSize;
-	const uint64 AlignmentPad  = AlignedSize       - Size;
+	const uint64 AlignedSize = FoundRange.Size - FindResult.LeftoverSize;
+	const uint64 AlignmentPad = AlignedSize - Size;
 	const uint64 AlignedOffset = FoundRange.Offset + AlignmentPad;
-	const uint64 AllocationEnd = AlignedOffset     + Size;
+	const uint64 AllocationEnd = AlignedOffset + Size;
 
 	// Adjust the range if there is space left over.
 	if (FindResult.LeftoverSize)
 	{
 		FoundRange.Offset = AllocationEnd;
-		FoundRange.Size  = FindResult.LeftoverSize;
+		FoundRange.Size = FindResult.LeftoverSize;
 	}
 	// Otherwise, remove it.
 	else
@@ -284,14 +327,13 @@ FRHITransientHeapAllocation FRHITransientHeapAllocator::Allocate(uint64 Size, ui
 	}
 
 	AllocationCount++;
-	UsedSize       += AlignedSize;
+	UsedSize += AlignedSize;
 	AlignmentWaste += AlignmentPad;
 
 	FRHITransientHeapAllocation Allocation;
-	Allocation.Size         = Size;
-	Allocation.Offset       = AlignedOffset;
+	Allocation.Size = Size;
+	Allocation.Offset = AlignedOffset;
 	Allocation.AlignmentPad = AlignmentPad;
-	Allocation.HeapIndex    = HeapIndex;
 
 	Validate();
 
@@ -304,12 +346,12 @@ void FRHITransientHeapAllocator::Deallocate(FRHITransientHeapAllocation Allocati
 
 	// Reconstruct the original range offset by subtracting the alignment pad, and expand the size accordingly.
 	const uint64 RangeToFreeOffset = Allocation.Offset - Allocation.AlignmentPad;
-	const uint64 RangeToFreeSize   = Allocation.Size   + Allocation.AlignmentPad;
-	const uint64 RangeToFreeEnd    = RangeToFreeOffset + RangeToFreeSize;
+	const uint64 RangeToFreeSize = Allocation.Size + Allocation.AlignmentPad;
+	const uint64 RangeToFreeEnd = RangeToFreeOffset + RangeToFreeSize;
 
 	FRangeHandle PreviousHandle = HeadHandle;
-	FRangeHandle NextHandle     = InvalidRangeHandle;
-	FRangeHandle Handle         = GetFirstFreeRangeHandle();
+	FRangeHandle NextHandle = InvalidRangeHandle;
+	FRangeHandle Handle = GetFirstFreeRangeHandle();
 
 	while (Handle != InvalidRangeHandle)
 	{
@@ -323,13 +365,13 @@ void FRHITransientHeapAllocator::Deallocate(FRHITransientHeapAllocation Allocati
 		}
 
 		PreviousHandle = Handle;
-		Handle         = Range.NextFreeHandle;
+		Handle = Range.NextFreeHandle;
 	}
 
 	uint64 MergedFreeRangeStart = RangeToFreeOffset;
-	uint64 MergedFreeRangeEnd   = RangeToFreeEnd;
-	bool bMergedPrevious        = false;
-	bool bMergedNext            = false;
+	uint64 MergedFreeRangeEnd = RangeToFreeEnd;
+	bool bMergedPrevious = false;
+	bool bMergedNext = false;
 
 	if (PreviousHandle != HeadHandle)
 	{
@@ -338,10 +380,10 @@ void FRHITransientHeapAllocator::Deallocate(FRHITransientHeapAllocation Allocati
 		// Attempt to merge the previous range with the range being freed.
 		if (PreviousRange.GetEnd() == RangeToFreeOffset)
 		{
-			PreviousRange.Size   += RangeToFreeSize;
-			MergedFreeRangeStart  = PreviousRange.Offset;
-			MergedFreeRangeEnd    = PreviousRange.GetEnd();
-			bMergedPrevious       = true;
+			PreviousRange.Size += RangeToFreeSize;
+			MergedFreeRangeStart = PreviousRange.Offset;
+			MergedFreeRangeEnd = PreviousRange.GetEnd();
+			bMergedPrevious = true;
 		}
 	}
 
@@ -352,11 +394,11 @@ void FRHITransientHeapAllocator::Deallocate(FRHITransientHeapAllocation Allocati
 		// Attempt to merge the next range with the range being freed.
 		if (RangeToFreeEnd == NextRange.Offset)
 		{
-			NextRange.Size      += RangeToFreeSize;
-			NextRange.Offset     = RangeToFreeOffset;
+			NextRange.Size += RangeToFreeSize;
+			NextRange.Offset = RangeToFreeOffset;
 			MergedFreeRangeStart = FMath::Min(MergedFreeRangeStart, RangeToFreeOffset);
-			MergedFreeRangeEnd   = NextRange.GetEnd();
-			bMergedNext          = true;
+			MergedFreeRangeEnd = NextRange.GetEnd();
+			bMergedNext = true;
 		}
 	}
 
@@ -364,7 +406,7 @@ void FRHITransientHeapAllocator::Deallocate(FRHITransientHeapAllocation Allocati
 	if (bMergedPrevious && bMergedNext)
 	{
 		FRange& PreviousRange = Ranges[PreviousHandle];
-		FRange& NextRange     = Ranges[NextHandle];
+		FRange& NextRange = Ranges[NextHandle];
 
 		PreviousRange.Size = MergedFreeRangeEnd - MergedFreeRangeStart;
 		RemoveRange(PreviousHandle, NextHandle);
@@ -375,7 +417,7 @@ void FRHITransientHeapAllocator::Deallocate(FRHITransientHeapAllocation Allocati
 		InsertRange(PreviousHandle, RangeToFreeOffset, RangeToFreeSize);
 	}
 
-	UsedSize       -= RangeToFreeSize;
+	UsedSize -= RangeToFreeSize;
 	AlignmentWaste -= Allocation.AlignmentPad;
 	AllocationCount--;
 
@@ -393,18 +435,18 @@ FRHITransientHeapAllocator::FFindResult FRHITransientHeapAllocator::FindFreeRang
 		FRange& Range = Ranges[Handle];
 
 		// Due to alignment we may have to shift the offset and expand the size accordingly.
-		const uint64 AlignmentPad = Align(HeapBaseGPUAddress + Range.Offset, Alignment) - HeapBaseGPUAddress - Range.Offset;
+		const uint64 AlignmentPad = Align(Range.Offset, Alignment) - Range.Offset;
 		const uint64 RequiredSize = Size + AlignmentPad;
 
 		if (RequiredSize <= Range.Size)
 		{
-			FindResult.FoundHandle  = Handle;
+			FindResult.FoundHandle = Handle;
 			FindResult.LeftoverSize = Range.Size - RequiredSize;
 			return FindResult;
 		}
 
 		FindResult.PreviousHandle = Handle;
-		Handle                    = Range.NextFreeHandle;
+		Handle = Range.NextFreeHandle;
 	}
 
 	return {};
@@ -412,17 +454,17 @@ FRHITransientHeapAllocator::FFindResult FRHITransientHeapAllocator::FindFreeRang
 
 void FRHITransientHeapAllocator::Validate()
 {
-#if RHICORE_TRANSIENT_ALLOCATOR_DEBUG
+#if UE_BUILD_DEBUG
 	uint64 DerivedFreeSize = 0;
 
 	FRangeHandle PreviousHandle = HeadHandle;
-	FRangeHandle NextHandle     = InvalidRangeHandle;
-	FRangeHandle Handle         = GetFirstFreeRangeHandle();
+	FRangeHandle NextHandle = InvalidRangeHandle;
+	FRangeHandle Handle = GetFirstFreeRangeHandle();
 
 	while (Handle != InvalidRangeHandle)
 	{
 		const FRange& Range = Ranges[Handle];
-		DerivedFreeSize    += Range.Size;
+		DerivedFreeSize += Range.Size;
 
 		if (PreviousHandle != HeadHandle)
 		{
@@ -433,273 +475,1167 @@ void FRHITransientHeapAllocator::Validate()
 		}
 
 		PreviousHandle = Handle;
-		Handle         = Range.NextFreeHandle;
+		Handle = Range.NextFreeHandle;
 	}
 
-	check(Initializer.Size == DerivedFreeSize + UsedSize);
+	check(Capacity == DerivedFreeSize + UsedSize);
 #endif
 }
 
 //////////////////////////////////////////////////////////////////////////
 
-void FRHITransientResourceOverlapTracker::Track(FRHITransientResource* TransientResource, uint64 OffsetMin, uint64 OffsetMax)
-{
-	check(TransientResource);
-
-	FResourceRange ResourceRangeNew;
-	ResourceRangeNew.Resource  = TransientResource;
-	ResourceRangeNew.OffsetMin = OffsetMin;
-	ResourceRangeNew.OffsetMax = OffsetMax;
-
-	for (int32 Index = 0; Index < ResourceRanges.Num(); ++Index)
-	{
-		FResourceRange& ResourceRangeOld = ResourceRanges[Index];
-
-		// If the old range starts later in the heap and doesn't overlap, the sort invariant guarantees no future range will overlap.
-		if (ResourceRangeOld.OffsetMin >= ResourceRangeNew.OffsetMax)
-		{
-			break;
-		}
-
-		// If the old range starts earlier in the heap and doesn't overlap, we keep searching.
-		if (ResourceRangeOld.OffsetMax <= ResourceRangeNew.OffsetMin)
-		{
-			continue;
-		}
-
-		TransientResource->AddAliasingOverlap(ResourceRangeOld.Resource);
-
-		// Complete overlap.
-		if (ResourceRangeOld.OffsetMin >= ResourceRangeNew.OffsetMin && ResourceRangeOld.OffsetMax <= ResourceRangeNew.OffsetMax)
-		{
-			ResourceRanges.RemoveAt(Index);
-			Index--;
-		}
-		// Partial overlap, can manifest as three cases:
-		else
-		{
-			bool bResizedOld = false;
-			FResourceRange ResourceRangeOldCopy = ResourceRangeOld;
-
-			// 1) New:    ********
-			//            |||        ->
-			//    Old: ======             ===********
-			if (ResourceRangeOld.OffsetMin < ResourceRangeNew.OffsetMin)
-			{
-				ResourceRangeOld.OffsetMax = ResourceRangeNew.OffsetMin;
-				bResizedOld = true;
-			}
-
-			if (ResourceRangeOldCopy.OffsetMax > ResourceRangeNew.OffsetMax)
-			{
-				// 2) New:    ********
-				//                |||      ->
-				//    Old:        ======         ********===
-				if (!bResizedOld)
-				{
-					ResourceRangeOld.OffsetMin = ResourceRangeNew.OffsetMax;
-				}
-
-				// 3) New:    ********
-				//            ||||||||      ->
-				//    Old: ==============        ===********===
-				else
-				{
-					// Lower bound has been resized already; add an upper bound.
-					FResourceRange ResourceRangeOldUpper = ResourceRangeOldCopy;
-					ResourceRangeOldUpper.OffsetMin = ResourceRangeNew.OffsetMax;
-					Index = ResourceRanges.Insert(ResourceRangeOldUpper, Index + 1);
-				}
-			}
-		}
-	}
-
-	for (int32 Index = 0; Index < ResourceRanges.Num(); ++Index)
-	{
-		const FResourceRange& ResourceRangeOld = ResourceRanges[Index];
-
-		// Insert new range while preserving order.
-		if (ResourceRangeOld.OffsetMin > ResourceRangeNew.OffsetMin)
-		{
-			ResourceRanges.Insert(ResourceRangeNew, Index);
-			return;
-		}
-	}
-
-	ResourceRanges.Add(ResourceRangeNew);
-}
-
-//////////////////////////////////////////////////////////////////////////
-
-FRHITransientResourceAllocator::FRHITransientResourceAllocator(FRHITransientResourceSystem& ParentInSystem)
-	: ParentSystem(ParentInSystem)
-{
-	Heaps.Reserve(ParentSystem.GetHeapCount());
-	HeapAllocators.Reserve(ParentSystem.GetHeapCount());
-}
-
-FRHITransientTexture* FRHITransientResourceAllocator::CreateTexture(
+FRHITransientTexture* FRHITransientHeap::CreateTexture(
 	const FRHITextureCreateInfo& CreateInfo,
 	const TCHAR* DebugName,
 	uint32 PassIndex,
+	uint64 CurrentAllocatorCycle,
 	uint64 TextureSize,
 	uint32 TextureAlignment,
 	FCreateTextureFunction CreateTextureFunction)
 {
-	if (TextureSize > ParentSystem.GetMaximumHeapSize())
+	FRHITransientHeapAllocation Allocation = Allocator.Allocate(TextureSize, TextureAlignment);
+	Allocation.Heap = this;
+
+	if (!Allocation.IsValid())
 	{
-		return {};
+		return nullptr;
 	}
 
-	const ERHITransientHeapFlags TextureHeapFlags =
-		EnumHasAnyFlags(CreateInfo.Flags, TexCreate_RenderTargetable | TexCreate_ResolveTargetable | TexCreate_DepthStencilTargetable | TexCreate_DepthStencilResolveTarget)
-		? ERHITransientHeapFlags::AllowRenderTargets
-		: ERHITransientHeapFlags::AllowTextures;
-
-	const FRHITransientHeapAllocation Allocation = Allocate(TextureStats, TextureSize, TextureAlignment, TextureHeapFlags);
-
-	FRHITransientHeapState& HeapState = Heaps[Allocation.HeapIndex];
-
-	FRHITransientTexture* TransientTexture = HeapState.Heap->AcquireTexture(CreateInfo, Allocation.Offset, [&] (uint64 Hash)
+	FRHITransientTexture* Texture = Textures.Acquire(ComputeHash(CreateInfo, Allocation.Offset), [&](uint64 Hash)
 	{
-		const FResourceInitializer ResourceInitializer(*HeapState.Heap, Allocation, Hash);
-		return CreateTextureFunction(ResourceInitializer);
+		Stats.Textures.CreateCount++;
+		return CreateTextureFunction(FResourceInitializer(Allocation, Hash));
 	});
 
-	InitResource(TransientTexture, PassIndex, Allocation, DebugName);
-
-#if RHICORE_TRANSIENT_ALLOCATOR_DEBUG
-	DebugTextures.Emplace(TransientTexture);
-	check(TransientTexture->GetCreateInfo() == CreateInfo && TransientTexture->GetName() == DebugName);
-#endif
+	check(Texture);
+	Texture->Acquire(DebugName, PassIndex, CurrentAllocatorCycle);
+	AllocateMemoryInternal(Texture, DebugName, PassIndex, CurrentAllocatorCycle, Allocation);
+	Stats.AllocateTexture(Allocation.Size);
+	return Texture;
+}
 
-	return MoveTemp(TransientTexture);
+void FRHITransientHeap::DeallocateMemory(FRHITransientTexture* Texture, uint32 PassIndex)
+{
+	DeallocateMemoryInternal(Texture, PassIndex);
+	Stats.DeallocateTexture(Texture->GetSize());
 }
 
-FRHITransientBuffer* FRHITransientResourceAllocator::CreateBuffer(
+FRHITransientBuffer* FRHITransientHeap::CreateBuffer(
 	const FRHIBufferCreateInfo& CreateInfo,
 	const TCHAR* DebugName,
 	uint32 PassIndex,
-	uint32 BufferSize,
+	uint64 CurrentAllocatorCycle,
+	uint64 BufferSize,
 	uint32 BufferAlignment,
 	FCreateBufferFunction CreateBufferFunction)
 {
-	if (BufferSize > ParentSystem.GetMaximumHeapSize())
+	FRHITransientHeapAllocation Allocation = Allocator.Allocate(BufferSize, BufferAlignment);
+	Allocation.Heap = this;
+
+	if (!Allocation.IsValid())
 	{
-		return {};
+		return nullptr;
 	}
 
-	const FRHITransientHeapAllocation Allocation = Allocate(BufferStats, BufferSize, BufferAlignment, ERHITransientHeapFlags::AllowBuffers);
-
-	FRHITransientHeapState& HeapState = Heaps[Allocation.HeapIndex];
-
-	FRHITransientBuffer* TransientBuffer = HeapState.Heap->AcquireBuffer(CreateInfo, Allocation.Offset, [&](uint64 Hash)
+	FRHITransientBuffer* Buffer = Buffers.Acquire(ComputeHash(CreateInfo, Allocation.Offset), [&](uint64 Hash)
 	{
-		const FResourceInitializer ResourceInitializer(*HeapState.Heap, Allocation, Hash);
-		return CreateBufferFunction(ResourceInitializer);
+		Stats.Buffers.CreateCount++;
+		return CreateBufferFunction(FResourceInitializer(Allocation, Hash));
 	});
 
-	InitResource(TransientBuffer, PassIndex, Allocation, DebugName);
-
-#if RHICORE_TRANSIENT_ALLOCATOR_DEBUG
-	DebugBuffers.Emplace(TransientBuffer);
-	check(TransientBuffer->GetCreateInfo() == CreateInfo && TransientBuffer->GetName() == DebugName);
-#endif
-
-	return TransientBuffer;
+	check(Buffer);
+	Buffer->Acquire(DebugName, PassIndex, CurrentAllocatorCycle);
+	AllocateMemoryInternal(Buffer, DebugName, PassIndex, CurrentAllocatorCycle, Allocation);
+	Stats.AllocateBuffer(Allocation.Size);
+	return Buffer;
 }
 
-void FRHITransientResourceAllocator::DeallocateMemoryInternal(FRHITransientResource* InResource, uint32 PassIndex, FMemoryStats& StatsToUpdate)
+void FRHITransientHeap::DeallocateMemory(FRHITransientBuffer* Buffer, uint32 PassIndex)
 {
-	InResource->SetDiscardPass(PassIndex);
+	DeallocateMemoryInternal(Buffer, PassIndex);
+	Stats.DeallocateBuffer(Buffer->GetSize());
+}
 
-	const FRHITransientHeapAllocation& Allocation = HeapAllocations[InResource->GetAllocationIndex()];
+void FRHITransientHeap::AllocateMemoryInternal(FRHITransientResource* Resource, const TCHAR* Name, uint32 PassIndex, uint64 CurrentAllocatorCycle, const FRHITransientHeapAllocation& Allocation)
+{
+	Resource->GetHeapAllocation() = Allocation;
 
-	HeapAllocators[Allocation.HeapIndex].Deallocate(Allocation);
+	check(Allocation.Offset % Initializer.Alignment == 0);
+	const uint64 AlignedSize   = Align(Allocation.Size, Initializer.Alignment);
+	const uint32 PageOffsetMin = Allocation.Offset >> AlignmentLog2;
+	const uint32 PageOffsetMax = (Allocation.Offset + AlignedSize) >> AlignmentLog2;
+	OverlapTracker.Track(Resource, PageOffsetMin, PageOffsetMax);
 
-	StatsToUpdate.CurrentSizeWithAliasing -= Allocation.Size;
+	CommitSize = FMath::Max(CommitSize, Allocation.Offset + Allocation.Size);
 }
 
-void FRHITransientResourceAllocator::Freeze(FRHICommandListImmediate& RHICmdList, FRHITransientHeapStats& OutHeapStats)
+void FRHITransientHeap::DeallocateMemoryInternal(FRHITransientResource* Resource, uint32 PassIndex)
 {
-	for (FRHITransientHeapState& State : Heaps)
-	{
-		auto& Heap = OutHeapStats.Heaps.AddDefaulted_GetRef();
-		Heap.WatermarkSize = State.CommitSize;
-		Heap.Capacity = State.Heap->GetCapacity();
-	}
+	Resource->Discard(PassIndex);
 
-	ParentSystem.ForfeitHeaps(Heaps);
+	const FRHITransientHeapAllocation Allocation = Resource->GetHeapAllocation();
+	check(Allocation.Heap == this);
 
-	RHICmdList.EnqueueLambda([&ParentSystem = ParentSystem, BufferStats = BufferStats, TextureStats = TextureStats](FRHICommandListImmediate&)
-	{
-		ParentSystem.TextureStats.Add(TextureStats);
-		ParentSystem.BufferStats.Add(BufferStats);
-	});
+	Allocator.Deallocate(Allocation);
 }
 
-void FRHITransientResourceAllocator::Flush()
+void FRHITransientHeap::Flush(uint64 AllocatorCycle, FRHITransientMemoryStats& OutMemoryStats, FRHITransientAllocationStats* OutAllocationStats)
 {
-	for (FRHITransientHeapState& State : Heaps)
+	const bool bHasDeallocations = Stats.HasDeallocations();
+	OutMemoryStats.Accumulate(Stats);
+	Stats.Reset();
+
+	if (OutAllocationStats)
 	{
-		State.LastCommitSize = State.CommitSize;
-	}
-}
+		const auto AddResourceToStats = [](FRHITransientAllocationStats& AllocationStats, FRHITransientResource* Resource)
+		{
+			const FRHITransientHeapAllocation& HeapAllocation = Resource->GetHeapAllocation();
 
-FRHITransientHeapAllocation FRHITransientResourceAllocator::Allocate(FMemoryStats& StatsToUpdate, uint64 Size, uint32 Alignment, ERHITransientHeapFlags ResourceHeapFlags)
-{
-	FRHITransientHeapAllocation Allocation;
+			FRHITransientAllocationStats::FAllocation Allocation;
+			Allocation.OffsetMin = HeapAllocation.Offset;
+			Allocation.OffsetMax = HeapAllocation.Offset + HeapAllocation.Size;
+			Allocation.MemoryRangeIndex = AllocationStats.MemoryRanges.Num();
 
-	for (int32 Index = 0; Index < HeapAllocators.Num(); ++Index)
-	{
-		FRHITransientHeapAllocator& HeapAllocator = HeapAllocators[Index];
+			AllocationStats.Resources.Emplace(Resource, FRHITransientAllocationStats::FAllocationArray{ Allocation });
+		};
 
-		if (!HeapAllocator.IsAllocationSupported(Size, ResourceHeapFlags))
+		OutAllocationStats->Resources.Reserve(OutAllocationStats->Resources.Num() + Textures.GetAllocatedCount() + Buffers.GetAllocatedCount());
+
+		for (FRHITransientTexture* Texture : Textures.GetAllocated())
 		{
-			continue;
+			AddResourceToStats(*OutAllocationStats, Texture);
 		}
 
-		Allocation = HeapAllocator.Allocate(Size, Alignment);
-
-		if (Allocation.IsValid())
+		for (FRHITransientBuffer* Buffer : Buffers.GetAllocated())
 		{
-			break;
+			AddResourceToStats(*OutAllocationStats, Buffer);
 		}
+
+		FRHITransientAllocationStats::FMemoryRange MemoryRange;
+		MemoryRange.Capacity = GetCapacity();
+		MemoryRange.CommitSize = GetCommitSize();
+		OutAllocationStats->MemoryRanges.Add(MemoryRange);
 	}
 
-	if (!Allocation.IsValid())
+	if (bHasDeallocations)
 	{
-		const uint32 HeapIndex = Heaps.Num();
+		CommitSize = 0;
+
+		{
+			Textures.Forfeit(GFrameCounterRenderThread);
+			LogLongLivedResources(Textures.GetAllocated(), AllocatorCycle);
 
-		FRHITransientHeapState State = {};
-		State.Heap = ParentSystem.AcquireHeap(Size, Alignment, ResourceHeapFlags);
-		Heaps.Emplace(State);
+			for (FRHITransientTexture* Texture : Textures.GetAllocated())
+			{
+				const FRHITransientHeapAllocation& Allocation = Texture->GetHeapAllocation();
+				CommitSize = FMath::Max(CommitSize, Allocation.Offset + Allocation.Size);
+			}
 
-		FRHITransientHeapAllocator& HeapAllocator = HeapAllocators.Emplace_GetRef(State.Heap->GetInitializer(), HeapIndex, State.Heap->GetBaseGPUVirtualAddress());
-		Allocation = HeapAllocator.Allocate(Size, Alignment);
+			TRACE_COUNTER_SET(TransientTextureCacheSize, Textures.GetSize());
+		}
 
-		check(Allocation.IsValid());
-	}
+		{
+			Buffers.Forfeit(GFrameCounterRenderThread);
+			LogLongLivedResources(Buffers.GetAllocated(), AllocatorCycle);
 
-	FRHITransientHeapState& State = Heaps[Allocation.HeapIndex];
-	State.CommitSize = FMath::Max(State.CommitSize, Allocation.Offset + Allocation.Size);
+			for (FRHITransientBuffer* Buffer : Buffers.GetAllocated())
+			{
+				const FRHITransientHeapAllocation& Allocation = Buffer->GetHeapAllocation();
+				CommitSize = FMath::Max(CommitSize, Allocation.Offset + Allocation.Size);
+			}
 
-	StatsToUpdate.AllocationCount++;
-	StatsToUpdate.TotalSize += Allocation.Size;
-	StatsToUpdate.CurrentSizeWithAliasing += Allocation.Size;
-	StatsToUpdate.TotalSizeWithAliasing = FMath::Max(StatsToUpdate.TotalSizeWithAliasing, StatsToUpdate.CurrentSizeWithAliasing);
+			TRACE_COUNTER_SET(TransientBufferCacheSize, Buffers.GetSize());
+		}
 
-	return Allocation;
+		OverlapTracker.Reset();
+	}
 }
 
-void FRHITransientResourceAllocator::InitResource(FRHITransientResource* TransientResource, uint32 PassIndex, const FRHITransientHeapAllocation& Allocation, const TCHAR* Name)
+//////////////////////////////////////////////////////////////////////////
+
+FRHITransientHeapCache::FInitializer FRHITransientHeapCache::FInitializer::CreateDefault()
 {
-	FRHITransientResourceStats Stats;
-	Stats.HeapIndex = Allocation.HeapIndex;
-	Stats.HeapOffsetMin = Allocation.Offset;
-	Stats.HeapOffsetMax = Allocation.Offset + Allocation.Size;
+	FInitializer Initializer;
+	Initializer.MinimumHeapSize = GRHITransientAllocatorMinimumHeapSize * 1024 * 1024;
+	Initializer.MaximumHeapSize = GRHITransientAllocatorMaximumHeapSize * 1024 * 1024;
+	Initializer.HeapAlignment = 64 * 1024;
+	Initializer.BufferCacheSize = GRHITransientAllocatorBufferCacheSize;
+	Initializer.TextureCacheSize = GRHITransientAllocatorTextureCacheSize;
+	Initializer.GarbageCollectLatency = GRHITransientAllocatorGarbageCollectLatency;
+	return Initializer;
+}
+
+FRHITransientHeapCache::~FRHITransientHeapCache()
+{
+	checkf(FreeList.Num() == LiveList.Num(), TEXT("Heaps are still being held by transient allocators."));
+	FreeList.Empty();
+	LiveList.Empty();
+}
+
+FRHITransientHeap* FRHITransientHeapCache::Acquire(uint64 FirstAllocationSize, ERHITransientHeapFlags FirstAllocationHeapFlags)
+{
+	FScopeLock Lock(&CriticalSection);
+
+	for (int32 HeapIndex = 0; HeapIndex < FreeList.Num(); ++HeapIndex)
+	{
+		FRHITransientHeap* Heap = FreeList[HeapIndex];
+
+		if (Heap->IsAllocationSupported(FirstAllocationSize, FirstAllocationHeapFlags))
+		{
+			FreeList.RemoveAt(HeapIndex);
+			return Heap;
+		}
+	}
+
+	FRHITransientHeap::FInitializer HeapInitializer;
+	HeapInitializer.Size = GetHeapSize(FirstAllocationSize);
+	HeapInitializer.Alignment = Initializer.HeapAlignment;
+	HeapInitializer.Flags = (Initializer.bSupportsAllHeapFlags ? ERHITransientHeapFlags::AllowAll : FirstAllocationHeapFlags);
+	HeapInitializer.TextureCacheSize = Initializer.TextureCacheSize;
+	HeapInitializer.BufferCacheSize = Initializer.BufferCacheSize;
+
+	LLM_SCOPE_BYTAG(RHITransientResources);
+	FRHITransientHeap* Heap = CreateHeap(HeapInitializer);
+	check(Heap);
+
+	TotalMemoryCapacity += HeapInitializer.Size;
+	LiveList.Emplace(Heap);
+	return Heap;
+}
+
+void FRHITransientHeapCache::Forfeit(TConstArrayView<FRHITransientHeap*> InForfeitedHeaps)
+{
+	FScopeLock Lock(&CriticalSection);
+
+	LiveList.Reserve(InForfeitedHeaps.Num());
+	for (FRHITransientHeap* Heap : InForfeitedHeaps)
+	{
+		check(Heap->IsEmpty());
+		Heap->LastUsedGarbageCollectCycle = GarbageCollectCycle;
+		FreeList.Add(Heap);
+	}
+
+	Algo::Sort(FreeList, [](const FRHITransientHeap* LHS, const FRHITransientHeap* RHS)
+	{
+		// Sort by smaller heap first.
+		if (LHS->GetCapacity() != RHS->GetCapacity())
+		{
+			return LHS->GetCapacity() < RHS->GetCapacity();
+		}
+
+		// Sort next by most recently used first.
+		return LHS->GetLastUsedGarbageCollectCycle() >= RHS->GetLastUsedGarbageCollectCycle();
+	});
+}
+
+void FRHITransientHeapCache::GarbageCollect()
+{
+	FScopeLock Lock(&CriticalSection);
+
+	for (int32 HeapIndex = 0; HeapIndex < FreeList.Num(); ++HeapIndex)
+	{
+		FRHITransientHeap* Heap = FreeList[HeapIndex];
+
+		if (Heap->GetLastUsedGarbageCollectCycle() + Initializer.GarbageCollectLatency <= GarbageCollectCycle)
+		{
+			TotalMemoryCapacity -= Heap->GetCapacity();
+			FreeList.RemoveAt(HeapIndex);
+			LiveList.Remove(Heap);
+			HeapIndex--;
+
+			delete Heap;
+		}
+	}
+
+	TRACE_COUNTER_SET(TransientMemoryRangeCount, LiveList.Num());
+
+	Stats.Submit(TotalMemoryCapacity);
+
+	GarbageCollectCycle++;
+}
+
+//////////////////////////////////////////////////////////////////////////
+
+FRHITransientResourceHeapAllocator::~FRHITransientResourceHeapAllocator()
+{
+	IF_RHICORE_TRANSIENT_ALLOCATOR_DEBUG(LogActiveResources(ActiveResources, CurrentCycle));
+}
+
+FRHITransientTexture* FRHITransientResourceHeapAllocator::CreateTextureInternal(
+	const FRHITextureCreateInfo& CreateInfo,
+	const TCHAR* DebugName,
+	uint32 PassIndex,
+	uint64 TextureSize,
+	uint32 TextureAlignment,
+	FRHITransientHeap::FCreateTextureFunction CreateTextureFunction)
+{
+	if (TextureSize > HeapCache.GetInitializer().MaximumHeapSize)
+	{
+		return nullptr;
+	}
+
+	const ERHITransientHeapFlags TextureHeapFlags =
+		EnumHasAnyFlags(CreateInfo.Flags, TexCreate_RenderTargetable | TexCreate_ResolveTargetable | TexCreate_DepthStencilTargetable | TexCreate_DepthStencilResolveTarget)
+		? ERHITransientHeapFlags::AllowRenderTargets
+		: ERHITransientHeapFlags::AllowTextures;
+
+	FRHITransientTexture* Texture = nullptr;
+
+	for (FRHITransientHeap* Heap : Heaps)
+	{
+		if (!Heap->IsAllocationSupported(TextureSize, TextureHeapFlags))
+		{
+			continue;
+		}
+
+		Texture = Heap->CreateTexture(CreateInfo, DebugName, PassIndex, CurrentCycle, TextureSize, TextureAlignment, CreateTextureFunction);
+
+		if (Texture)
+		{
+			break;
+		}
+	}
+
+	if (!Texture)
+	{
+		FRHITransientHeap* Heap = HeapCache.Acquire(TextureSize, TextureHeapFlags);
+		Heaps.Emplace(Heap);
+
+		Texture = Heap->CreateTexture(CreateInfo, DebugName, PassIndex, CurrentCycle, TextureSize, TextureAlignment, CreateTextureFunction);
+	}
+
+	check(Texture);
+	IF_RHICORE_TRANSIENT_ALLOCATOR_DEBUG(ActiveResources.Emplace(Texture));
+	return Texture;
+}
+
+FRHITransientBuffer* FRHITransientResourceHeapAllocator::CreateBufferInternal(
+	const FRHIBufferCreateInfo& CreateInfo,
+	const TCHAR* DebugName,
+	uint32 PassIndex,
+	uint32 BufferSize,
+	uint32 BufferAlignment,
+	FRHITransientHeap::FCreateBufferFunction CreateBufferFunction)
+{
+	if (BufferSize > HeapCache.GetInitializer().MaximumHeapSize)
+	{
+		return nullptr;
+	}
+
+	FRHITransientBuffer* Buffer = nullptr;
+
+	for (FRHITransientHeap* Heap : Heaps)
+	{
+		if (!Heap->IsAllocationSupported(BufferSize, ERHITransientHeapFlags::AllowBuffers))
+		{
+			continue;
+		}
+
+		Buffer = Heap->CreateBuffer(CreateInfo, DebugName, PassIndex, CurrentCycle, BufferSize, BufferAlignment, CreateBufferFunction);
+
+		if (Buffer)
+		{
+			break;
+		}
+	}
+
+	if (!Buffer)
+	{
+		FRHITransientHeap* Heap = HeapCache.Acquire(BufferSize, ERHITransientHeapFlags::AllowBuffers);
+		Heaps.Emplace(Heap);
+
+		Buffer = Heap->CreateBuffer(CreateInfo, DebugName, PassIndex, CurrentCycle, BufferSize, BufferAlignment, CreateBufferFunction);
+	}
+
+	check(Buffer);
+	IF_RHICORE_TRANSIENT_ALLOCATOR_DEBUG(ActiveResources.Emplace(Buffer));
+	return Buffer;
+}
+
+void FRHITransientResourceHeapAllocator::DeallocateMemory(FRHITransientTexture* Texture, uint32 PassIndex)
+{
+	check(Texture);
+
+	FRHITransientHeap* Heap = Texture->GetHeapAllocation().Heap;
+
+	check(Heap);
+	check(Heaps.Contains(Heap));
+
+	Heap->DeallocateMemory(Texture, PassIndex);
+	DeallocationCount++;
+
+	IF_RHICORE_TRANSIENT_ALLOCATOR_DEBUG(ActiveResources.Remove(Texture));
+}
+
+void FRHITransientResourceHeapAllocator::DeallocateMemory(FRHITransientBuffer* Buffer, uint32 PassIndex)
+{
+	check(Buffer);
+
+	FRHITransientHeap* Heap = Buffer->GetHeapAllocation().Heap;
+
+	check(Heap);
+	check(Heaps.Contains(Heap));
+
+	Heap->DeallocateMemory(Buffer, PassIndex);
+	DeallocationCount++;
+
+	IF_RHICORE_TRANSIENT_ALLOCATOR_DEBUG(ActiveResources.Remove(Buffer));
+}
+
+void FRHITransientResourceHeapAllocator::Flush(FRHICommandListImmediate& RHICmdList, FRHITransientAllocationStats* OutAllocationStats)
+{
+	FRHITransientMemoryStats Stats;
+
+	for (FRHITransientHeap* Heap : Heaps)
+	{
+		Heap->Flush(CurrentCycle, Stats, OutAllocationStats);
+	}
+
+	if (DeallocationCount > 0)
+	{
+		int32 FirstForfeitIndex = Algo::Partition(Heaps.GetData(), Heaps.Num(), [](const FRHITransientHeap* Heap) { return !Heap->IsEmpty(); });
+		HeapCache.Forfeit(MakeArrayView(Heaps.GetData() + FirstForfeitIndex, Heaps.Num() - FirstForfeitIndex));
+		Heaps.SetNum(FirstForfeitIndex, false);
+		DeallocationCount = 0;
+	}
+
+	RHICmdList.EnqueueLambda([&HeapCache = HeapCache, Stats](FRHICommandListImmediate&)
+	{
+		HeapCache.Stats.Accumulate(Stats);
+	});
+
+	CurrentCycle++;
+}
+
+///////////////////////////////////////////////////////////////////////////////////////////////////
+//! Transient Resource Page Allocator
+///////////////////////////////////////////////////////////////////////////////////////////////////
+
+void FRHITransientPageSpanAllocator::Init()
+{
+	check(MaxSpanCount == MaxPageCount + 2);
+
+	PageToSpanStart.AddDefaulted(MaxPageCount + 1);
+	PageToSpanEnd.AddDefaulted(MaxPageCount + 1);
+
+	PageSpans.AddDefaulted(MaxSpanCount);
+	UnusedSpanList.AddDefaulted(MaxSpanCount);
+
+	Reset();
+}
+
+void FRHITransientPageSpanAllocator::Reset()
+{
+	FreePageCount = MaxPageCount;
+	AllocationCount = 0;
+
+	// Initialize the unused span index pool with MaxSpanCount entries
+	for (uint32 Index = 0; Index < MaxSpanCount; Index++)
+	{
+		UnusedSpanList[Index] = MaxSpanCount - 1 - Index;
+	}
+	UnusedSpanListCount = MaxSpanCount;
+
+	// Allocate the head and tail spans (dummy spans), and a span between them covering the entire range
+	uint32 HeadSpanIndex = AllocSpan();
+	uint32 TailSpanIndex = AllocSpan();
+	check(HeadSpanIndex == FreeSpanListHeadIndex);
+	check(TailSpanIndex == FreeSpanListTailIndex);
+
+	if (MaxPageCount > 0)
+	{
+		uint32 FirstFreeNodeIndex = AllocSpan();
+
+		// Allocate head and tail nodes (0 and 1)
+		for (uint32 Index = 0; Index < 2; Index++)
+		{
+			PageSpans[Index].Offset = 0;
+			PageSpans[Index].Count = 0;
+			PageSpans[Index].PrevSpanIndex = InvalidIndex;
+			PageSpans[Index].NextSpanIndex = InvalidIndex;
+			PageSpans[Index].bAllocated = false;
+		}
+		PageSpans[HeadSpanIndex].NextSpanIndex = FirstFreeNodeIndex;
+		PageSpans[TailSpanIndex].PrevSpanIndex = FirstFreeNodeIndex;
+
+		// First Node
+		PageSpans[FirstFreeNodeIndex].Offset = 0;
+		PageSpans[FirstFreeNodeIndex].Count = MaxPageCount;
+		PageSpans[FirstFreeNodeIndex].PrevSpanIndex = HeadSpanIndex;
+		PageSpans[FirstFreeNodeIndex].NextSpanIndex = TailSpanIndex;
+		PageSpans[FirstFreeNodeIndex].bAllocated = false;
+
+		// Initialize the page->span mapping
+		for (uint32 Index = 0; Index < MaxPageCount + 1; Index++)
+		{
+			PageToSpanStart[Index] = InvalidIndex;
+			PageToSpanEnd[Index] = InvalidIndex;
+		}
+		PageToSpanStart[0] = FirstFreeNodeIndex;
+		PageToSpanEnd[MaxPageCount] = FirstFreeNodeIndex;
+	}
+	else
+	{
+		PageSpans[HeadSpanIndex].NextSpanIndex = TailSpanIndex;
+		PageSpans[TailSpanIndex].PrevSpanIndex = HeadSpanIndex;
+	}
+}
+
+bool FRHITransientPageSpanAllocator::Allocate(uint32 PageCount, uint32& OutNumPagesAllocated, uint32& OutSpanIndex)
+{
+	OutNumPagesAllocated = 0;
+
+	if (FreePageCount < PageCount)
+	{
+		// If we're allowing partial allocs and we run out of pages, allocate all the remaining pages
+		PageCount = FreePageCount;
+	}
+
+	if (PageCount > FreePageCount || PageCount == 0)
+	{
+		return false;
+	}
+	OutNumPagesAllocated = PageCount;
+
+	// Allocate spans from the free list head
+	uint32 NumPagesToFind = PageCount;
+	uint32 FoundPages = 0;
+	FPageSpan& HeadSpan = PageSpans[FreeSpanListHeadIndex];
+	uint32 StartSpanIndex = HeadSpan.NextSpanIndex;
+	uint32 SpanIndex = StartSpanIndex;
+	while (SpanIndex != FreeSpanListTailIndex && SpanIndex != InvalidIndex)
+	{
+		FPageSpan& Span = PageSpans[SpanIndex];
+		if (NumPagesToFind <= Span.Count)
+		{
+			// Span is too big, so split it
+			if (Span.Count > NumPagesToFind)
+			{
+				SplitSpan(SpanIndex, NumPagesToFind);
+			}
+			check(NumPagesToFind == Span.Count);
+
+			// Move the head to point to the next free span
+			if (HeadSpan.NextSpanIndex != InvalidIndex)
+			{
+				PageSpans[HeadSpan.NextSpanIndex].PrevSpanIndex = InvalidIndex;
+			}
+			HeadSpan.NextSpanIndex = Span.NextSpanIndex;
+			if (Span.NextSpanIndex != InvalidIndex)
+			{
+				PageSpans[Span.NextSpanIndex].PrevSpanIndex = FreeSpanListHeadIndex;
+			}
+			Span.NextSpanIndex = InvalidIndex;
+		}
+		Span.bAllocated = true;
+		NumPagesToFind -= Span.Count;
+		SpanIndex = Span.NextSpanIndex;
+	}
+	check(NumPagesToFind == 0);
+	FreePageCount -= PageCount;
+#if UE_BUILD_DEBUG
+	Validate();
+#endif
+	AllocationCount++;
+	OutSpanIndex = StartSpanIndex;
+	return true;
+}
+
+void FRHITransientPageSpanAllocator::Deallocate(uint32 SpanIndex)
+{
+	if (SpanIndex == InvalidIndex)
+	{
+		return;
+	}
+	check(AllocationCount > 0);
+	// Find the right span with which to merge this
+	while (SpanIndex != InvalidIndex)
+	{
+		FPageSpan& FreedSpan = PageSpans[SpanIndex];
+		check(FreedSpan.bAllocated);
+		FreePageCount += FreedSpan.Count;
+		uint32 NextSpanIndex = FreedSpan.NextSpanIndex;
+		FreedSpan.bAllocated = false;
+		if (!MergeFreeSpanIfPossible(SpanIndex))
+		{
+			// If we can't merge this span, just unlink and add it to the head (or tail)
+			Unlink(SpanIndex);
+
+			InsertAfter(FreeSpanListHeadIndex, SpanIndex);
+		}
+		SpanIndex = NextSpanIndex;
+	}
+	AllocationCount--;
+
+#if UE_BUILD_DEBUG
+	Validate();
+#endif
+}
+
+void FRHITransientPageSpanAllocator::SplitSpan(uint32 InSpanIndex, uint32 InPageCount)
+{
+	FPageSpan& Span = PageSpans[InSpanIndex];
+	check(InPageCount <= Span.Count);
+	if (InPageCount < Span.Count)
+	{
+		uint32 NewSpanIndex = AllocSpan();
+		FPageSpan& NewSpan = PageSpans[NewSpanIndex];
+		NewSpan.NextSpanIndex = Span.NextSpanIndex;
+		NewSpan.PrevSpanIndex = InSpanIndex;
+		NewSpan.Count = Span.Count - InPageCount;
+		NewSpan.Offset = Span.Offset + InPageCount;
+		NewSpan.bAllocated = Span.bAllocated;
+		Span.Count = InPageCount;
+		Span.NextSpanIndex = NewSpanIndex;
+		if (NewSpan.NextSpanIndex != InvalidIndex)
+		{
+			PageSpans[NewSpan.NextSpanIndex].PrevSpanIndex = NewSpanIndex;
+		}
+
+		// Update the PageToSpan mappings
+		PageToSpanEnd[NewSpan.Offset] = InSpanIndex;
+		PageToSpanStart[NewSpan.Offset] = NewSpanIndex;
+		PageToSpanEnd[NewSpan.Offset + NewSpan.Count] = NewSpanIndex;
+	}
+}
+
+void FRHITransientPageSpanAllocator::MergeSpans(uint32 SpanIndex0, uint32 SpanIndex1, const bool bKeepSpan1)
+{
+	FPageSpan& Span0 = PageSpans[SpanIndex0];
+	FPageSpan& Span1 = PageSpans[SpanIndex1];
+	check(Span0.Offset + Span0.Count == Span1.Offset);
+	check(Span0.bAllocated == Span1.bAllocated);
+	check(Span0.NextSpanIndex == SpanIndex1);
+	check(Span1.PrevSpanIndex == SpanIndex0);
+
+	uint32 SpanIndexToKeep = bKeepSpan1 ? SpanIndex1 : SpanIndex0;
+	uint32 SpanIndexToRemove = bKeepSpan1 ? SpanIndex0 : SpanIndex1;
+
+	// Update the PageToSpan mappings
+	PageToSpanStart[Span0.Offset] = SpanIndexToKeep;
+	PageToSpanStart[Span1.Offset] = InvalidIndex;
+	PageToSpanEnd[Span0.Offset + Span0.Count] = InvalidIndex; // Should match Span1.Offset
+	PageToSpanEnd[Span1.Offset + Span1.Count] = SpanIndexToKeep;
+	if (bKeepSpan1)
+	{
+		Span1.Offset = Span0.Offset;
+		Span1.Count += Span0.Count;
+	}
+	else
+	{
+		Span0.Count += Span1.Count;
+	}
+
+	Unlink(SpanIndexToRemove);
+	ReleaseSpan(SpanIndexToRemove);
+}
+
+// Inserts a span after an existing span. The span to insert must be unlinked
+void FRHITransientPageSpanAllocator::InsertAfter(uint32 InsertPosition, uint32 InsertSpanIndex)
+{
+	check(InsertPosition != InvalidIndex);
+	check(InsertSpanIndex != InvalidIndex);
+	FPageSpan& SpanAtPos = PageSpans[InsertPosition];
+	FPageSpan& SpanToInsert = PageSpans[InsertSpanIndex];
+	check(!SpanToInsert.IsLinked());
+
+	// Connect Span0's next node with the inserted node
+	SpanToInsert.NextSpanIndex = SpanAtPos.NextSpanIndex;
+	if (SpanAtPos.NextSpanIndex != InvalidIndex)
+	{
+		PageSpans[SpanAtPos.NextSpanIndex].PrevSpanIndex = InsertSpanIndex;
+	}
+	// Connect the two nodes
+	SpanAtPos.NextSpanIndex = InsertSpanIndex;
+	SpanToInsert.PrevSpanIndex = InsertPosition;
+}
+
+// Inserts a span after an existing span. The span to insert must be unlinked
+void FRHITransientPageSpanAllocator::InsertBefore(uint32 InsertPosition, uint32 InsertSpanIndex)
+{
+	check(InsertPosition != InvalidIndex && InsertPosition != 0); // Can't insert before the head
+	check(InsertSpanIndex != InvalidIndex);
+	FPageSpan& SpanAtPos = PageSpans[InsertPosition];
+	FPageSpan& SpanToInsert = PageSpans[InsertSpanIndex];
+	check(!SpanToInsert.IsLinked());
+
+	// Connect Span0's prev node with the inserted node
+	SpanToInsert.PrevSpanIndex = SpanAtPos.PrevSpanIndex;
+	if (SpanAtPos.PrevSpanIndex != InvalidIndex)
+	{
+		PageSpans[SpanAtPos.PrevSpanIndex].NextSpanIndex = InsertSpanIndex;
+	}
+
+	// Connect the two nodes
+	SpanAtPos.PrevSpanIndex = InsertSpanIndex;
+	SpanToInsert.NextSpanIndex = InsertPosition;
+}
+
+void FRHITransientPageSpanAllocator::Unlink(uint32 SpanIndex)
+{
+	FPageSpan& Span = PageSpans[SpanIndex];
+	check(SpanIndex != FreeSpanListHeadIndex);
+	if (Span.PrevSpanIndex != InvalidIndex)
+	{
+		PageSpans[Span.PrevSpanIndex].NextSpanIndex = Span.NextSpanIndex;
+	}
+	if (Span.NextSpanIndex != InvalidIndex)
+	{
+		PageSpans[Span.NextSpanIndex].PrevSpanIndex = Span.PrevSpanIndex;
+	}
+	Span.PrevSpanIndex = InvalidIndex;
+	Span.NextSpanIndex = InvalidIndex;
+}
+
+uint32 FRHITransientPageSpanAllocator::GetAllocationPageCount(uint32 SpanIndex) const
+{
+	check(SpanIndex != InvalidIndex && SpanIndex < MaxSpanCount);
+	check(PageSpans[SpanIndex].bAllocated);
+
+	uint32 Count = 0;
+	do
+	{
+		Count += PageSpans[SpanIndex].Count;
+		SpanIndex = PageSpans[SpanIndex].NextSpanIndex;
+
+	} while (SpanIndex != InvalidIndex);
+
+	return Count;
+}
+
+bool FRHITransientPageSpanAllocator::MergeFreeSpanIfPossible(uint32 SpanIndex)
+{
+	FPageSpan& Span = PageSpans[SpanIndex];
+	check(!Span.bAllocated);
+	bool bMerged = false;
+
+	// Can we merge this span with an existing one to the left?
+	uint32 AdjSpanIndexPrev = PageToSpanEnd[Span.Offset];
+	if (AdjSpanIndexPrev != InvalidIndex && !PageSpans[AdjSpanIndexPrev].bAllocated)
+	{
+		Unlink(SpanIndex);
+		InsertAfter(AdjSpanIndexPrev, SpanIndex);
+		MergeSpans(AdjSpanIndexPrev, SpanIndex, true);
+		bMerged = true;
+	}
+
+	// Can we merge this span with an existing free one to the right?
+	uint32 AdjSpanIndexNext = PageToSpanStart[Span.Offset + Span.Count];
+	if (AdjSpanIndexNext != InvalidIndex && !PageSpans[AdjSpanIndexNext].bAllocated)
+	{
+		Unlink(SpanIndex);
+		InsertBefore(AdjSpanIndexNext, SpanIndex);
+		MergeSpans(SpanIndex, AdjSpanIndexNext, false);
+		bMerged = true;
+	}
+
+	return bMerged;
+}
+
+void FRHITransientPageSpanAllocator::Validate()
+{
+#if UE_BUILD_DEBUG
+	// Check the mappings are valid
+	for (uint32 Index = 0; Index < MaxPageCount; Index++)
+	{
+		check(PageToSpanStart[Index] == InvalidIndex || PageSpans[PageToSpanStart[Index]].Offset == Index);
+		check(PageToSpanEnd[Index] == InvalidIndex || PageSpans[PageToSpanEnd[Index]].Offset + PageSpans[PageToSpanEnd[Index]].Count == Index);
+	}
+
+	// Count free pages
+	uint32 FreeCount = 0;
+	int32 PrevIndex = FreeSpanListHeadIndex;
+	for (int32 Index = GetFirstSpanIndex(); Index != InvalidIndex; Index = PageSpans[Index].NextSpanIndex)
+	{
+		FPageSpan& Span = PageSpans[Index];
+		check(Span.PrevSpanIndex == PrevIndex);
+		PrevIndex = Index;
+		FreeCount += Span.Count;
+	}
+	check(FreeCount <= MaxPageCount);
+	check(FreeCount == FreePageCount);
+#endif
+}
+
+///////////////////////////////////////////////////////////////////////////////////////////////////
+
+void FRHITransientPagePool::Allocate(FAllocationContext& Context)
+{
+	uint32 SpanIndex = 0;
+	uint32 PagesAllocated = 0;
+	if (Allocator.Allocate(Context.PagesRemaining, PagesAllocated, SpanIndex))
+	{
+		const uint64 DestinationGpuVirtualAddress = Context.GpuVirtualAddress + Context.PagesAllocated * Initializer.PageSize;
+		const uint32 PageSpanOffsetMin = PageSpans.Num();
+
+		Allocator.GetSpanArray(SpanIndex, PageSpans);
+
+		const uint32 PageSpanOffsetMax = PageSpans.Num();
+		const uint32 PageSpanCount     = PageSpanOffsetMax - PageSpanOffsetMin;
+
+		const  int32 AllocationIndex   = Context.Allocations.Num();
+		const uint64 AllocationHash    = CityHash64WithSeed((const char*)&PageSpans[PageSpanOffsetMin], PageSpanCount * sizeof(FRHITransientPageSpan), DestinationGpuVirtualAddress);
+
+		for (uint32 Index = PageSpanOffsetMin; Index < PageSpanOffsetMax; ++Index)
+		{
+			const FRHITransientPageSpan Span = PageSpans[Index];
+			OverlapTracker.Track(&Context.Resource, Span.Offset, Span.Offset + Span.Count);
+		}
+
+		FRHITransientPagePoolAllocation Allocation;
+		Allocation.Pool = this;
+		Allocation.Hash = AllocationHash;
+		Allocation.SpanOffsetMin = Context.Spans.Num();
+		Allocation.SpanOffsetMax = Context.Spans.Num() + PageSpanCount;
+		Allocation.SpanIndex = SpanIndex;
+		Context.Allocations.Emplace(Allocation);
+		Context.Spans.Append(&PageSpans[PageSpanOffsetMin], PageSpanCount);
+
+		bool bMapPages = true;
+
+#if 0
+		if (AllocationIndex < Context.AllocationsBefore.Num())
+		{
+			const FRHITransientPagePoolAllocation& AllocationBefore = Context.AllocationsBefore[AllocationIndex];
+
+			if (AllocationBefore.Hash == AllocationHash && AllocationBefore.Pool == this)
+			{
+				Context.AllocationMatchingCount++;
+				bMapPages = false;
+			}
+		}
+#endif
+
+		if (bMapPages)
+		{
+			PageMapRequests.Emplace(DestinationGpuVirtualAddress, GpuVirtualAddress, Initializer.PageCount, PageSpanOffsetMin, PageSpanCount);
+			Context.PagesMapped += PagesAllocated;
+		}
+
+		check(Context.PagesRemaining >= PagesAllocated);
+		Context.PagesRemaining -= PagesAllocated;
+		Context.PagesAllocated += PagesAllocated;
+		Context.PageSpansAllocated += PageSpanCount;
+		Context.AllocationCount++;
+	}
+}
+
+void FRHITransientPagePool::Flush(FRHICommandListImmediate& RHICmdList)
+{
+	PageMapRequestCountMax = FMath::Max<uint32>(PageMapRequests.Num(), PageMapRequestCountMax);
+	PageSpanCountMax       = FMath::Max<uint32>(PageSpans.Num(),       PageSpanCountMax);
+
+	Flush(RHICmdList, MoveTemp(PageMapRequests), MoveTemp(PageSpans));
+
+	PageMapRequests.Reserve(PageMapRequestCountMax);
+	PageSpans.Reserve(PageSpanCountMax);
+
+	OverlapTracker.Reset();
+}
+
+///////////////////////////////////////////////////////////////////////////////////////////////////
+
+FRHITransientPagePoolCache::FInitializer FRHITransientPagePoolCache::FInitializer::CreateDefault()
+{
+	FInitializer Initializer;
+	Initializer.BufferCacheSize = GRHITransientAllocatorBufferCacheSize;
+	Initializer.TextureCacheSize = GRHITransientAllocatorTextureCacheSize;
+	Initializer.GarbageCollectLatency = GRHITransientAllocatorGarbageCollectLatency;
+	return Initializer;
+}
+
+FRHITransientPagePoolCache::~FRHITransientPagePoolCache()
+{
+	checkf(FreeList.Num() == LiveList.Num(), TEXT("PagePools are still being held by transient allocators."));
+	FreeList.Empty();
+	LiveList.Empty();
+}
+
+FRHITransientPagePool* FRHITransientPagePoolCache::Acquire()
+{
+	FScopeLock Lock(&CriticalSection);
+
+	if (!FreeList.IsEmpty())
+	{
+		return FreeList.Pop();
+	}
+
+	LLM_SCOPE_BYTAG(RHITransientResources);
+
+	FRHITransientPagePool::FInitializer PagePoolInitializer;
+	PagePoolInitializer.PageSize = Initializer.PageSize;
+
+	if (LiveList.IsEmpty() && Initializer.PoolSizeFirst > Initializer.PoolSize)
+	{
+		PagePoolInitializer.PageCount = Initializer.PoolSizeFirst / Initializer.PageSize;
+	}
+	else
+	{
+		PagePoolInitializer.PageCount = Initializer.PoolSize / Initializer.PageSize;
+	}
+
+	FRHITransientPagePool* PagePool = CreatePagePool(PagePoolInitializer);
+	check(PagePool);
+
+	TotalMemoryCapacity += PagePool->GetCapacity();
+	LiveList.Emplace(PagePool);
+
+	return PagePool;
+}
+
+void FRHITransientPagePoolCache::Forfeit(TConstArrayView<FRHITransientPagePool*> InForfeitedPagePools)
+{
+	FScopeLock Lock(&CriticalSection);
+
+	// These are iterated in reverse so they are acquired in the same order. 
+	for (int32 Index = InForfeitedPagePools.Num() - 1; Index >= 0; --Index)
+	{
+		FRHITransientPagePool* PagePool = InForfeitedPagePools[Index];
+		check(PagePool->IsEmpty());
+		PagePool->LastUsedGarbageCollectCycle = GarbageCollectCycle;
+		FreeList.Add(PagePool);
+	}
+}
+
+void FRHITransientPagePoolCache::GarbageCollect()
+{
+	SCOPED_NAMED_EVENT_TEXT("TransientPagePoolCache::GarbageCollect", FColor::Magenta);
+	TArray<FRHITransientPagePool*, TInlineAllocator<16>> PoolsToDelete;
+
+	{
+		FScopeLock Lock(&CriticalSection);
+
+		for (int32 PagePoolIndex = 0; PagePoolIndex < FreeList.Num(); ++PagePoolIndex)
+		{
+			FRHITransientPagePool* PagePool = FreeList[PagePoolIndex];
+
+			if (PagePool->GetLastUsedGarbageCollectCycle() + Initializer.GarbageCollectLatency <= GarbageCollectCycle)
+			{
+				TotalMemoryCapacity -= PagePool->GetCapacity();
+				FreeList.RemoveAt(PagePoolIndex);
+				LiveList.Remove(PagePool);
+				PagePoolIndex--;
+
+				PoolsToDelete.Emplace(PagePool);
+
+			#if 1 // Only delete one per frame. Deletion can be quite expensive.
+				break;
+			#endif
+			}
+		}
+
+		TRACE_COUNTER_SET(TransientMemoryRangeCount, LiveList.Num());
+	}
+
+	Stats.Submit(TotalMemoryCapacity);
+
+	GarbageCollectCycle++;
+
+	for (FRHITransientPagePool* PagePool : PoolsToDelete)
+	{
+		delete PagePool;
+	}
+}
+
+///////////////////////////////////////////////////////////////////////////////////////////////////
+
+FRHITransientResourcePageAllocator::~FRHITransientResourcePageAllocator()
+{
+	IF_RHICORE_TRANSIENT_ALLOCATOR_DEBUG(LogActiveResources(ActiveResources, CurrentCycle));
+}
+
+FRHITransientTexture* FRHITransientResourcePageAllocator::CreateTexture(
+	const FRHITextureCreateInfo& CreateInfo,
+	const TCHAR* DebugName,
+	uint32 PassIndex)
+{
+	FRHITransientTexture* Texture = Textures.Acquire(ComputeHash(CreateInfo), [&](uint64 Hash)
+	{
+		Stats.Textures.CreateCount++;
+		return CreateTextureInternal(CreateInfo, DebugName, Hash);
+	});
+
+	check(Texture);
+	Texture->Acquire(DebugName, PassIndex, CurrentCycle);
+	AllocateMemoryInternal(Texture, DebugName, PassIndex, EnumHasAnyFlags(CreateInfo.Flags, ETextureCreateFlags::FastVRAM));
+	Stats.AllocateTexture(Texture->GetSize());
+	IF_RHICORE_TRANSIENT_ALLOCATOR_DEBUG(ActiveResources.Emplace(Texture));
+	return Texture;
+}
+
+FRHITransientBuffer* FRHITransientResourcePageAllocator::CreateBuffer(
+	const FRHIBufferCreateInfo& CreateInfo,
+	const TCHAR* DebugName,
+	uint32 PassIndex)
+{
+	FRHITransientBuffer* Buffer = Buffers.Acquire(ComputeHash(CreateInfo), [&](uint64 Hash)
+	{
+		Stats.Buffers.CreateCount++;
+		return CreateBufferInternal(CreateInfo, DebugName, Hash);
+	});
+
+	check(Buffer);
+	Buffer->Acquire(DebugName, PassIndex, CurrentCycle);
+	AllocateMemoryInternal(Buffer, DebugName, PassIndex, EnumHasAnyFlags(CreateInfo.Usage, EBufferUsageFlags::FastVRAM));
+	Stats.AllocateBuffer(Buffer->GetSize());
+	IF_RHICORE_TRANSIENT_ALLOCATOR_DEBUG(ActiveResources.Emplace(Buffer));
+	return Buffer;
+}
+
+void FRHITransientResourcePageAllocator::AllocateMemoryInternal(FRHITransientResource* Resource, const TCHAR* DebugName, uint32 PassIndex, bool bFastPoolRequested)
+{
+	FRHITransientPagePool::FAllocationContext AllocationContext(*Resource, PageSize);
+
+	if (bFastPoolRequested)
+	{
+		if (FRHITransientPagePool* FastPagePool = PagePoolCache.GetFastPagePool())
+		{
+			FastPagePool->Allocate(AllocationContext);
+		}
+	}
+
+	if (!AllocationContext.IsComplete())
+	{
+		for (FRHITransientPagePool* PagePool : PagePools)
+		{
+			PagePool->Allocate(AllocationContext);
+
+			if (AllocationContext.IsComplete())
+			{
+				break;
+			}
+		}
+	}
+
+	while (!AllocationContext.IsComplete())
+	{
+		FRHITransientPagePool* PagePool = PagePoolCache.Acquire();
+		PagePool->Allocate(AllocationContext);
+		PagePools.Emplace(PagePool);
+	}
+
+	PageMapCount      += AllocationContext.PagesMapped;
+	PageAllocateCount += AllocationContext.PagesAllocated;
+	PageSpanCount     += AllocationContext.PageSpansAllocated;
+}
+
+void FRHITransientResourcePageAllocator::DeallocateMemoryInternal(FRHITransientResource* Resource, uint32 PassIndex)
+{
+	Resource->Discard(PassIndex);
+
+	for (const FRHITransientPagePoolAllocation& Allocation : Resource->GetPageAllocation().PoolAllocations)
+	{
+		Allocation.Pool->Deallocate(Allocation.SpanIndex);
+	}
+}
+
+void FRHITransientResourcePageAllocator::DeallocateMemory(FRHITransientTexture* Texture, uint32 PassIndex)
+{
+	DeallocateMemoryInternal(Texture, PassIndex);
+	Stats.DeallocateTexture(Texture->GetSize());
+}
+
+void FRHITransientResourcePageAllocator::DeallocateMemory(FRHITransientBuffer* Buffer, uint32 PassIndex)
+{
+	DeallocateMemoryInternal(Buffer, PassIndex);
+	Stats.DeallocateBuffer(Buffer->GetSize());
+}
+
+void FRHITransientResourcePageAllocator::Flush(FRHICommandListImmediate& RHICmdList, FRHITransientAllocationStats* OutAllocationStats)
+{
+	if (OutAllocationStats)
+	{
+		TMap<FRHITransientPagePool*, int32> PagePoolToIndex;
+		PagePoolToIndex.Reserve(PagePools.Num());
+
+		for (int32 Index = 0; Index < PagePools.Num(); ++Index)
+		{
+			FRHITransientPagePool* PagePool = PagePools[Index];
+
+			FRHITransientAllocationStats::FMemoryRange MemoryRange;
+			MemoryRange.Capacity = PagePool->GetCapacity();
+			MemoryRange.CommitSize = PagePool->GetCapacity();
+			OutAllocationStats->MemoryRanges.Emplace(MemoryRange);
+
+			PagePoolToIndex.Emplace(PagePool, Index);
+		}
+
+		for (const FRHITransientTexture* Texture : Textures.GetAllocated())
+		{
+			FRHITransientAllocationStats::FAllocationArray& Allocations = OutAllocationStats->Resources.Emplace(Texture);
+
+			EnumeratePageSpans(Texture, [&](FRHITransientPagePool* PagePool, FRHITransientPageSpan PageSpan)
+			{
+				FRHITransientAllocationStats::FAllocation Allocation;
+				Allocation.OffsetMin = PageSize *  PageSpan.Offset;
+				Allocation.OffsetMax = PageSize * (PageSpan.Offset + PageSpan.Count);
+				Allocation.MemoryRangeIndex = PagePoolToIndex[PagePool];
+				Allocations.Emplace(Allocation);
+			});
+		}
+
+		for (const FRHITransientBuffer* Buffer : Buffers.GetAllocated())
+		{
+			FRHITransientAllocationStats::FAllocationArray& Allocations = OutAllocationStats->Resources.Emplace(Buffer);
+
+			EnumeratePageSpans(Buffer, [&](FRHITransientPagePool* PagePool, FRHITransientPageSpan PageSpan)
+			{
+				FRHITransientAllocationStats::FAllocation Allocation;
+				Allocation.OffsetMin = PageSize *  PageSpan.Offset;
+				Allocation.OffsetMax = PageSize * (PageSpan.Offset + PageSpan.Count);
+				Allocation.MemoryRangeIndex = PagePoolToIndex[PagePool];
+				Allocations.Emplace(Allocation);
+			});
+		}
+	}
+
+	{
+		TRACE_COUNTER_SET(TransientPageMapCount, PageMapCount);
+		TRACE_COUNTER_SET(TransientPageAllocateCount, PageAllocateCount);
+		TRACE_COUNTER_SET(TransientPageSpanCount, PageSpanCount);
+		PageMapCount = 0;
+		PageAllocateCount = 0;
+		PageSpanCount = 0;
+	}
+
+	{
+		TRACE_COUNTER_SET(TransientTextureCount, Textures.GetAllocatedCount());
+		TRACE_COUNTER_SET(TransientTextureCacheHitPercentage, Textures.GetHitPercentage());
+
+		Textures.Forfeit(GFrameCounterRenderThread, [this](FRHITransientTexture* Texture) { ReleaseTextureInternal(Texture); });
+
+		LogLongLivedResources(Textures.GetAllocated(), CurrentCycle);
+		TRACE_COUNTER_SET(TransientTextureCacheSize, Textures.GetSize());
+	}
+
+	{
+		TRACE_COUNTER_SET(TransientBufferCount, Buffers.GetAllocatedCount());
+		TRACE_COUNTER_SET(TransientBufferCacheHitPercentage, Buffers.GetHitPercentage());
+
+		Buffers.Forfeit(GFrameCounterRenderThread, [this](FRHITransientBuffer* Buffer) { ReleaseBufferInternal(Buffer); });
+
+		LogLongLivedResources(Buffers.GetAllocated(), CurrentCycle);
+		TRACE_COUNTER_SET(TransientBufferCacheSize, Buffers.GetSize());
+	}
+
+	for (FRHITransientPagePool* PagePool : PagePools)
+	{
+		PagePool->Flush(RHICmdList);
+	}
+
+	if (Stats.HasDeallocations())
+	{
+		const int32 FirstForfeitIndex = Algo::Partition(PagePools.GetData(), PagePools.Num(), [](const FRHITransientPagePool* PagePool) { return !PagePool->IsEmpty(); });
+		PagePoolCache.Forfeit(MakeArrayView(PagePools.GetData() + FirstForfeitIndex, PagePools.Num() - FirstForfeitIndex));
+		PagePools.SetNum(FirstForfeitIndex, false);
+	}
+
+	RHICmdList.EnqueueLambda([&PagePoolCache = PagePoolCache, Stats = Stats](FRHICommandListImmediate&)
+	{
+		PagePoolCache.Stats.Accumulate(Stats);
+	});
 
-	TransientResource->Init(Name, HeapAllocations.Emplace(Allocation), PassIndex, Stats);
+	Stats.Reset();
 
-	HeapAllocators[Allocation.HeapIndex].TrackOverlap(TransientResource, Allocation);
+	CurrentCycle++;
 }
diff --git a/Engine/Source/Runtime/RHICore/Public/RHICoreTransientResourceAllocator.h b/Engine/Source/Runtime/RHICore/Public/RHICoreTransientResourceAllocator.h
--- a/Engine/Source/Runtime/RHICore/Public/RHICoreTransientResourceAllocator.h
+++ b/Engine/Source/Runtime/RHICore/Public/RHICoreTransientResourceAllocator.h
@@ -2,9 +2,7 @@
 #pragma once
 
 #include "RHITransientResourceAllocator.h"
-
-class FRHITransientResourceSystem;
-class FRHITransientResourceAllocator;
+#include "Algo/Partition.h"
 
 #define RHICORE_TRANSIENT_ALLOCATOR_DEBUG (!UE_BUILD_SHIPPING && !UE_BUILD_TEST)
 
@@ -23,30 +21,144 @@ inline uint64 ComputeHash(const FRHITextureCreateInfo& InCreateInfo, uint64 Heap
 	return CityHash64WithSeed((const char*)&NewInfo, sizeof(FRHITextureCreateInfo), HeapOffset);
 }
 
+inline uint64 ComputeHash(const FRHITextureCreateInfo& InCreateInfo)
+{
+	// Make sure all padding is removed.
+	FRHITextureCreateInfo NewInfo;
+	FPlatformMemory::Memzero(&NewInfo, sizeof(FRHITextureCreateInfo));
+	NewInfo = InCreateInfo;
+	return CityHash64((const char*)&NewInfo, sizeof(FRHITextureCreateInfo));
+}
+
 inline uint64 ComputeHash(const FRHIBufferCreateInfo& InCreateInfo, uint64 HeapOffset)
 {
 	return CityHash64WithSeed((const char*)&InCreateInfo, sizeof(FRHIBufferCreateInfo), HeapOffset);
 }
 
-struct FRHITransientMemoryStats
+inline uint64 ComputeHash(const FRHIBufferCreateInfo& InCreateInfo)
+{
+	return CityHash64((const char*)&InCreateInfo, sizeof(FRHIBufferCreateInfo));
+}
+
+/** Tracks allocation statistics for buffer or texture resources. */
+struct FRHITransientResourceStats
+{
+	FRHITransientResourceStats() = default;
+
+	void Add(const FRHITransientResourceStats& Other)
+	{
+		AllocatedSize     += Other.AllocatedSize;
+		AllocationCount   += Other.AllocationCount;
+		DeallocationCount += Other.DeallocationCount;
+		CreateCount       += Other.CreateCount;
+	}
+
+	void Allocate(uint64 Size)
+	{
+		AllocationCount++;
+		AllocatedSize      += Size;
+	}
+
+	void Deallocate(uint64 Size)
+	{
+		DeallocationCount++;
+	}
+
+	// The number of bytes allocated from the transient allocator this cycle.
+	uint64 AllocatedSize = 0;
+
+	// The number of allocations made from the transient allocator this cycle.
+	uint32 AllocationCount = 0;
+
+	// the number of deallocations made from the transient allocator this cycle.
+	uint32 DeallocationCount = 0;
+
+	// The number of resource creations made from the transient allocator this cycle.
+	uint32 CreateCount = 0;
+};
+
+/** Tracks all transient memory statistics for the current allocation cycle and reports results to various profilers. */
+struct RHICORE_API FRHITransientMemoryStats
 {
+public:
 	FRHITransientMemoryStats() = default;
 
-	void Add(const FRHITransientMemoryStats& Other)
+	void Accumulate(const FRHITransientMemoryStats& Other)
+	{
+		Textures.Add(Other.Textures);
+		Buffers.Add(Other.Buffers);
+		AliasedSize = FMath::Max(AliasedSize, Other.AliasedSize);
+	}
+
+	void AllocateTexture(uint64 Size)
 	{
-		TotalSizeWithAliasing += Other.TotalSizeWithAliasing;
-		TotalSize += Other.TotalSize;
-		AllocationCount += Other.AllocationCount;
+		Textures.Allocate(Size);
+		AliasedSizeCurrent += Size;
+		AliasedSize         = FMath::Max(AliasedSize, AliasedSizeCurrent);
 	}
 
+	void DeallocateTexture(uint64 Size)
+	{
+		Textures.Deallocate(Size);
+		AliasedSizeCurrent -= Size;
+	}
+
+	void AllocateBuffer(uint64 Size)
+	{
+		Buffers.Allocate(Size);
+		AliasedSizeCurrent += Size;
+		AliasedSize         = FMath::Max(AliasedSize, AliasedSizeCurrent);
+	}
+
+	void DeallocateBuffer(uint64 Size)
+	{
+		Buffers.Deallocate(Size);
+		AliasedSizeCurrent -= Size;
+	}
+
+	void Reset()
+	{
+		Textures = {};
+		Buffers = {};
+		AliasedSize = AliasedSizeCurrent;
+	}
+
+	void Submit(uint64 TotalMemoryCapacity);
+
+	bool HasDeallocations() const { return Textures.DeallocationCount > 0 || Buffers.DeallocationCount > 0; }
+
+	FRHITransientResourceStats Textures;
+	FRHITransientResourceStats Buffers;
+
 	// Total allocated memory usage with aliasing.
-	uint64 TotalSizeWithAliasing = 0;
+	uint64 AliasedSizeCurrent = 0;
 
-	// Total allocated memory usage without aliasing.
-	uint64 TotalSize = 0;
+	// Current aliased size as items are being allocated / deallocated.
+	uint64 AliasedSize = 0;
+};
 
-	// The number of allocations made from the transient allocator.
-	uint32 AllocationCount = 0;
+/** Tracks resource allocations on the heap and adds overlap events to transient resources. */
+class RHICORE_API FRHITransientResourceOverlapTracker
+{
+public:
+	void Track(FRHITransientResource* InResource, uint32 PageOffsetMin, uint32 PageOffsetMax);
+
+	void Reset()
+	{
+		ResourceRanges.Reset();
+		Resources.Reset();
+	}
+
+private:
+	struct FResourceRange
+	{
+		uint64 PageOffsetMin : 24;
+		uint64 PageOffsetMax : 24;
+		uint64 ResourceIndex : 16;
+	};
+
+	TArray<FResourceRange> ResourceRanges;
+	TArray<FRHITransientResource*> Resources;
 };
 
 /** An RHI transient resource cache designed to optimize fetches for resources placed into a heap with an offset.
@@ -62,7 +174,9 @@ public:
 	static const uint32 kDefaultCapacity = kInfinity;
 	static const uint32 kDefaultGarbageCollectLatency = 32;
 
-	TRHITransientResourceCache(uint32 InCapacity = kDefaultCapacity, uint32 InGarbageCollectLatency = kDefaultGarbageCollectLatency)
+	TRHITransientResourceCache(
+		uint32 InCapacity = kDefaultCapacity,
+		uint32 InGarbageCollectLatency = kDefaultGarbageCollectLatency)
 		: GarbageCollectLatency(InGarbageCollectLatency)
 		, Capacity(InCapacity)
 	{
@@ -78,6 +192,11 @@ public:
 		{
 			delete Item.Resource;
 		}
+
+		for (TransientResourceType* Resource : Allocated)
+		{
+			delete Resource;
+		}
 	}
 
 	template <typename CreateFunctionType>
@@ -90,47 +209,77 @@ public:
 			if (CacheItem.Hash == Hash)
 			{
 				TransientResourceType* Resource = CacheItem.Resource;
-				Cache.RemoveAtSwap(Index);
+				Cache.RemoveAtSwap(Index, 1, false);
+				Allocated.Emplace(Resource);
+				HitCount++;
 				return Resource;
 			}
 		}
 
-		return CreateFunction(Hash);
+		TransientResourceType* Resource = CreateFunction(Hash);
+		Allocated.Emplace(Resource);
+		MissCount++;
+		return Resource;
 	}
 
-	void Forfeit(TArrayView<TransientResourceType* const> Resources)
+	template <typename ReleaseFunctionType>
+	void Forfeit(uint64 CurrentFrameIndex, ReleaseFunctionType ReleaseFunction)
 	{
-		CurrentGeneration++;
+		const int32 FirstForfeitIndex = Algo::Partition(Allocated.GetData(), Allocated.Num(), [](const FRHITransientResource* Resource) { return Resource->IsAcquired(); });
+		const auto ResourcesToForfeit = MakeArrayView(Allocated.GetData() + FirstForfeitIndex, Allocated.Num() - FirstForfeitIndex);
 
-		for (TransientResourceType* Resource : Resources)
+		for (TransientResourceType* Resource : ResourcesToForfeit)
 		{
-			Cache.Emplace(Resource, Resource->GetHash(), CurrentGeneration);
+			Cache.Emplace(Resource, Resource->GetHash(), CurrentFrameIndex);
 		}
 
+		Allocated.SetNum(FirstForfeitIndex, false);
+
 		Algo::Sort(Cache, [](const FCacheItem& LHS, const FCacheItem& RHS)
 		{
-			return LHS.Generation > RHS.Generation;
+			return LHS.LastUsedFrame > RHS.LastUsedFrame;
 		});
 
 		while (uint32(Cache.Num()) > Capacity)
 		{
-			if (!TryReleaseItem())
+			if (!TryReleaseItem(CurrentFrameIndex, ReleaseFunction))
 			{
 				break;
 			}
 		}
+
+		HitCount = 0;
+		MissCount = 0;
+	}
+
+	void Forfeit(uint64 CurrentFrameIndex)
+	{
+		Forfeit(CurrentFrameIndex, [](TransientResourceType*) {});
 	}
 
-	FORCEINLINE uint32 GetSize() const { return Cache.Num(); }
-	FORCEINLINE uint32 GetCapacity() const { return Capacity; }
+	TConstArrayView<TransientResourceType*> GetAllocated() const { return Allocated; }
+
+	uint32 GetAllocatedCount() const { return Allocated.Num(); }
+
+	uint32 GetSize() const { return Cache.Num(); }
+
+	uint32 GetCapacity() const { return Capacity; }
+
+	uint32 GetHitCount() const { return HitCount; }
+
+	uint32 GetMissCount() const { return MissCount; }
+
+	float GetHitPercentage() const { return (float)HitCount / (float)(HitCount + MissCount); }
 
 private:
-	bool TryReleaseItem()
+	template <typename ReleaseFunctionType>
+	bool TryReleaseItem(uint64 CurrentFrameIndex, ReleaseFunctionType ReleaseFunction)
 	{
 		const FCacheItem& Item = Cache.Top();
 
-		if (Item.Generation + GarbageCollectLatency < CurrentGeneration)
+		if (Item.LastUsedFrame + GarbageCollectLatency <= CurrentFrameIndex)
 		{
+			ReleaseFunction(Item.Resource);
 			delete Item.Resource;
 			Cache.Pop();
 			return true;
@@ -141,23 +290,142 @@ private:
 
 	struct FCacheItem
 	{
-		FCacheItem(TransientResourceType* InResource, uint64 InHash, uint64 InGeneration)
+		FCacheItem(TransientResourceType* InResource, uint64 InHash, uint64 InLastUsedFrame)
 			: Resource(InResource)
 			, Hash(InHash)
-			, Generation(InGeneration)
+			, LastUsedFrame(InLastUsedFrame)
 		{}
 
 		TransientResourceType* Resource;
 		uint64 Hash{};
-		uint64 Generation{};
+		uint64 LastUsedFrame{};
 	};
 
 	TArray<FCacheItem> Cache;
-	uint64 CurrentGeneration{};
+	TArray<TransientResourceType*> Allocated;
 	uint32 GarbageCollectLatency;
 	uint32 Capacity;
+	uint32 HitCount = 0;
+	uint32 MissCount = 0;
 };
 
+///////////////////////////////////////////////////////////////////////////////////////////////////
+
+/** Abstract base class for memory caching, providing a common garbage collection method. */
+class IRHITransientMemoryCache
+{
+public:
+	virtual ~IRHITransientMemoryCache() = default;
+
+	virtual void GarbageCollect() = 0;
+};
+
+///////////////////////////////////////////////////////////////////////////////////////////////////
+//! Transient Resource Heap Allocator
+///////////////////////////////////////////////////////////////////////////////////////////////////
+
+class FRHITransientHeapCache;
+class FRHITransientResourceHeapAllocator;
+
+/** First-fit allocator used for placing resources on a heap. */
+class RHICORE_API FRHITransientHeapAllocator
+{
+public:
+	FRHITransientHeapAllocator(uint64 Capacity, uint32 Alignment);
+	~FRHITransientHeapAllocator();
+
+	FRHITransientHeapAllocation Allocate(uint64 Size, uint32 Alignment);
+
+	void Deallocate(FRHITransientHeapAllocation Allocation);
+
+	uint64 GetCapacity() const { return Capacity; }
+	uint64 GetUsedSize() const { return UsedSize; }
+	uint64 GetFreeSize() const { return Capacity - UsedSize; }
+	uint64 GetAlignmentWaste() const { return AlignmentWaste; }
+	uint32 GetAllocationCount() const { return AllocationCount; }
+
+	bool IsFull() const { return UsedSize == Capacity; }
+	bool IsEmpty() const { return UsedSize == 0; }
+
+private:
+	using FRangeHandle = uint16;
+	static const FRangeHandle InvalidRangeHandle = FRangeHandle(~0);
+
+	struct FRange
+	{
+		uint64 Size{};
+		uint64 Offset{};
+		FRangeHandle NextFreeHandle = InvalidRangeHandle;
+
+		FORCEINLINE uint64 GetStart() const { return Offset; }
+		FORCEINLINE uint64 GetEnd() const { return Size + Offset; }
+	};
+
+	FORCEINLINE FRangeHandle GetFirstFreeRangeHandle()
+	{
+		return Ranges[HeadHandle].NextFreeHandle;
+	}
+
+	FRangeHandle CreateRange()
+	{
+		if (!RangeFreeList.IsEmpty())
+		{
+			return RangeFreeList.Pop();
+		}
+		Ranges.Emplace();
+		return FRangeHandle(Ranges.Num() - 1);
+	}
+
+	FRangeHandle InsertRange(FRangeHandle PreviousHandle, uint64 Offset, uint64 Size)
+	{
+		FRangeHandle Handle = CreateRange();
+
+		FRange& CurrentRange = Ranges[Handle];
+		CurrentRange.Offset = Offset;
+		CurrentRange.Size = Size;
+
+		FRange& PreviousRange = Ranges[PreviousHandle];
+		CurrentRange.NextFreeHandle = PreviousRange.NextFreeHandle;
+		PreviousRange.NextFreeHandle = Handle;
+
+		return Handle;
+	}
+
+	void RemoveRange(FRangeHandle PreviousHandle, FRangeHandle CurrentHandle)
+	{
+		FRange& PreviousRange = Ranges[PreviousHandle];
+		FRange& CurrentRange = Ranges[CurrentHandle];
+
+		PreviousRange.NextFreeHandle = CurrentRange.NextFreeHandle;
+		CurrentRange.NextFreeHandle = InvalidRangeHandle;
+
+		RangeFreeList.Add(CurrentHandle);
+	}
+
+	struct FFindResult
+	{
+		uint64 LeftoverSize = 0;
+		FRangeHandle PreviousHandle = InvalidRangeHandle;
+		FRangeHandle FoundHandle = InvalidRangeHandle;
+	};
+
+	FFindResult FindFreeRange(uint64 Size, uint32 Alignment);
+
+	void Validate();
+
+	uint64 Capacity = 0;
+	uint64 UsedSize = 0;
+	uint64 AlignmentWaste = 0;
+	uint32 AllocationCount = 0;
+	uint32 AlignmentMin = 0;
+
+	FRangeHandle HeadHandle = InvalidRangeHandle;
+	TArray<FRangeHandle> RangeFreeList;
+	TArray<FRange> Ranges;
+};
+
+///////////////////////////////////////////////////////////////////////////////////////////////////
+
 enum class ERHITransientHeapFlags : uint8
 {
 	// Supports placing buffers onto the heap.
@@ -175,24 +443,6 @@ enum class ERHITransientHeapFlags : uint8
 
 ENUM_CLASS_FLAGS(ERHITransientHeapFlags);
 
-struct FRHITransientHeapInitializer
-{
-	// Size of the heap in bytes.
-	uint64 Size{};
-
-	// Alignment of the heap in bytes.
-	uint32 Alignment{};
-
-	// Flags used to filter resource allocations within the heap.
-	ERHITransientHeapFlags Flags = ERHITransientHeapFlags::AllowAll;
-
-	// Size of the texture cache before elements are evicted.
-	uint32 TextureCacheSize{};
-
-	// Size of the buffer cache before elements are evicted.
-	uint32 BufferCacheSize{};
-};
-
 /** The base class for a platform heap implementation. Transient resources are placed on the heap at specific
  *  byte offsets. Each heap additionally contains a cache of RHI transient resources, each with its own RHI
  *  resource and cache of RHI views. The lifetime of the resource cache is tied to the heap.
@@ -200,442 +450,676 @@ struct FRHITransientHeapInitializer
 class RHICORE_API FRHITransientHeap
 {
 public:
-	FRHITransientHeap(const FRHITransientHeapInitializer& InInitializer)
-		: Initializer(InInitializer)
-		, Textures(InInitializer.TextureCacheSize)
-		, Buffers(InInitializer.BufferCacheSize)
-	{}
+	struct FInitializer
+	{
+		// Size of the heap in bytes.
+		uint64 Size = 0;
 
-	virtual ~FRHITransientHeap() = default;
+		// Alignment of the heap in bytes.
+		uint32 Alignment = 0;
+
+		// Flags used to filter resource allocations within the heap.
+		ERHITransientHeapFlags Flags = ERHITransientHeapFlags::AllowAll;
+
+		// Size of the texture cache before elements are evicted.
+		uint32 TextureCacheSize = 0;
+
+		// Size of the buffer cache before elements are evicted.
+		uint32 BufferCacheSize = 0;
+	};
+
+	struct FResourceInitializer
+	{
+		FResourceInitializer(const FRHITransientHeapAllocation& InAllocation, uint64 InHash)
+			: Heap(*InAllocation.Heap)
+			, Allocation(InAllocation)
+			, Hash(InHash)
+		{}
 
-	FORCEINLINE const FRHITransientHeapInitializer& GetInitializer() const { return Initializer; }
+		// The heap on which to create the resource.
+		FRHITransientHeap& Heap;
 
-	FORCEINLINE uint64 GetBaseGPUVirtualAddress() const { return BaseGPUVirtualAddress; }
+		// The allocation (offset / size) on the provided heap.
+		const FRHITransientHeapAllocation& Allocation;
 
-	FORCEINLINE uint64 GetCapacity() const { return Initializer.Size; }
+		// The unique hash computed from the create info and allocation offset.
+		const uint64 Hash;
+	};
 
-	FORCEINLINE uint64 GetLastUsedGarbageCollectCycle() const { return LastUsedGarbageCollectCycle; }
+	using FCreateTextureFunction = TFunction<FRHITransientTexture* (const FResourceInitializer&)>;
+	using FCreateBufferFunction = TFunction<FRHITransientBuffer* (const FResourceInitializer&)>;
 
-	FORCEINLINE bool IsAllocationSupported(uint64 Size, uint32 Alignment, ERHITransientHeapFlags InFlags) const
+	FRHITransientHeap(const FInitializer& InInitializer)
+		: Initializer(InInitializer)
+		, Allocator(InInitializer.Size, InInitializer.Alignment)
+		, AlignmentLog2(FPlatformMath::CeilLogTwo64(InInitializer.Alignment))
+		, Textures(InInitializer.TextureCacheSize)
+		, Buffers(InInitializer.BufferCacheSize)
 	{
-		return Size <= Initializer.Size && Alignment <= Initializer.Alignment && EnumHasAnyFlags(Initializer.Flags, InFlags);
+		check(1ull << AlignmentLog2 == InInitializer.Alignment);
 	}
 
-protected:
-	uint64 BaseGPUVirtualAddress = 0;
+	virtual ~FRHITransientHeap() = default;
 
-private:
-	template <typename CreateFunctionType>
-	FRHITransientTexture* AcquireTexture(const FRHITextureCreateInfo& CreateInfo, uint64 HeapOffset, CreateFunctionType CreateFunction);
+	FRHITransientTexture* CreateTexture(
+		const FRHITextureCreateInfo& CreateInfo,
+		const TCHAR* DebugName,
+		uint32 PassIndex,
+		uint64 CurrentAllocatorCycle,
+		uint64 TextureSize,
+		uint32 TextureAlignment,
+		FCreateTextureFunction CreateTextureFunction);
 
-	template <typename CreateFunctionType>
-	FRHITransientBuffer* AcquireBuffer(const FRHIBufferCreateInfo& CreateInfo, uint64 HeapOffset, CreateFunctionType CreateFunction);
+	void DeallocateMemory(FRHITransientTexture* Texture, uint32 PassIndex);
 
-	void ForfeitResources();
+	FRHITransientBuffer* CreateBuffer(
+		const FRHIBufferCreateInfo& CreateInfo,
+		const TCHAR* DebugName,
+		uint32 PassIndex,
+		uint64 CurrentAllocatorCycle,
+		uint64 BufferSize,
+		uint32 BufferAlignment,
+		FCreateBufferFunction CreateBufferFunction);
 
-	FRHITransientHeapInitializer Initializer;
-	uint64 LastUsedGarbageCollectCycle = 0;
+	void DeallocateMemory(FRHITransientBuffer* Buffer, uint32 PassIndex);
 
-	TRHITransientResourceCache<FRHITransientTexture> Textures;
-	TRHITransientResourceCache<FRHITransientBuffer> Buffers;
-	TArray<FRHITransientTexture*> AllocatedTextures;
-	TArray<FRHITransientBuffer*> AllocatedBuffers;
+	void Flush(uint64 CurrentAllocatorCycle, FRHITransientMemoryStats& OutMemoryStats, FRHITransientAllocationStats* OutAllocationStats);
 
-	friend FRHITransientResourceSystem;
-	friend FRHITransientResourceAllocator;
-};
+	const FInitializer& GetInitializer() const { return Initializer; }
 
-struct RHICORE_API FRHITransientResourceSystemInitializer
-{
-	// Creates a default initializer using common RHI CVars.
-	static FRHITransientResourceSystemInitializer CreateDefault();
+	uint64 GetCapacity() const { return Allocator.GetCapacity(); }
 
-	static const uint32 kDefaultResourceCacheSize = 256;
+	uint64 GetGPUVirtualAddress() const { return GpuVirtualAddress; }
 
-	// The minimum size to use when creating a heap. This is the default but can grow based on allocations.
-	uint64 MinimumHeapSize{};
+	uint64 GetLastUsedGarbageCollectCycle() const { return LastUsedGarbageCollectCycle; }
 
-	// The maximum size of a pool. Allocations above this size will fail.
-	uint64 MaximumHeapSize{};
+	uint64 GetCommitSize() const { return CommitSize; }
 
-	// The alignment of all heaps in the cache.
-	uint32 HeapAlignment{};
+	bool IsEmpty() const { return Allocator.IsEmpty(); }
 
-	// The latency between the completed fence value and the used fence value to invoke garbage collection of the heap.
-	uint32 GarbageCollectLatency{};
+	bool IsFull() const { return Allocator.IsFull(); }
 
-	// Size of the texture cache before elements are evicted.
-	uint32 TextureCacheSize = kDefaultResourceCacheSize;
+	bool IsCommitRequired() const { return CommitSize > 0; }
 
-	// Size of the buffer cache before elements are evicted.
-	uint32 BufferCacheSize = kDefaultResourceCacheSize;
+	bool IsAllocationSupported(uint64 Size, ERHITransientHeapFlags Flags) const
+	{
+		return Size <= Allocator.GetFreeSize() && EnumHasAnyFlags(Initializer.Flags, Flags);
+	}
 
-	// Whether all heaps should be created with the AllowAll heap flag.
-	bool bSupportsAllHeapFlags = true;
-};
+protected:
+	void SetGPUVirtualAddress(uint64 InBaseGPUVirtualAddress)
+	{
+		GpuVirtualAddress = InBaseGPUVirtualAddress;
+	}
 
-struct FRHITransientHeapState
-{
-	FRHITransientHeap* Heap;
+private:
+	void AllocateMemoryInternal(FRHITransientResource* Resource, const TCHAR* Name, uint32 PassIndex, uint64 CurrentAllocatorCycle, const FRHITransientHeapAllocation& Allocation);
+	void DeallocateMemoryInternal(FRHITransientResource* Resource, uint32 PassIndex);
+
+	FInitializer Initializer;
+	FRHITransientHeapAllocator Allocator;
 
-	// The total number of bytes used by an instance of
-	// FRHITransientResourceAllocator for the given heap.
+	uint64 GpuVirtualAddress = 0;
+	uint64 LastUsedGarbageCollectCycle = 0;
 	uint64 CommitSize = 0;
+	uint32 AlignmentLog2;
 
-	// The commit size since the last flush operation.
-	uint64 LastCommitSize = 0;
+	FRHITransientResourceOverlapTracker OverlapTracker;
 
-	// Returns whether a commit is required for the heap since the last flush. This is reset with each Flush() call on the allocator.
-	bool IsCommitRequired() const { return CommitSize > LastCommitSize; }
+	FRHITransientMemoryStats Stats;
+	TRHITransientResourceCache<FRHITransientTexture> Textures;
+	TRHITransientResourceCache<FRHITransientBuffer> Buffers;
+
+	friend FRHITransientHeapCache;
+	friend FRHITransientResourceHeapAllocator;
 };
 
-/** The RHI transient resource system is a base class for the platform implementation. It has a persistent lifetime
+///////////////////////////////////////////////////////////////////////////////////////////////////
+
+/** The RHI transient heap system is a base class for the platform implementation. It has a persistent lifetime
  *  and contains a cache of transient heaps. The transient allocator acquires heaps from the system and forfeits them
  *  at the end of its lifetime. Garbage collection of heaps is done using an internal counter that increments with
  *  each GarbageCollect call. This should be done periodically. Heaps older than a platform-specified fence latency
- *  are destroyed. Additionally, statistics are gathered automatically and reported to the platform implementation for
- *  stats tracking. However, 'rhitransientmemory' stats are reported automatically by the implementation itself.
+ *  are destroyed. Additionally, statistics are gathered automatically and reported to the 'rhitransientmemory' stats
+ *  group.
  */
-class RHICORE_API FRHITransientResourceSystem
+class RHICORE_API FRHITransientHeapCache : public IRHITransientMemoryCache
 {
 public:
-	FRHITransientResourceSystem(const FRHITransientResourceSystemInitializer& InInitializer)
+	struct RHICORE_API FInitializer
+	{
+		// Creates a default initializer using common RHI CVars.
+		static FInitializer CreateDefault();
+
+		static const uint32 kDefaultResourceCacheSize = 256;
+
+		// The minimum size to use when creating a heap. This is the default but can grow based on allocations.
+		uint64 MinimumHeapSize = 0;
+
+		// The maximum size of a pool. Allocations above this size will fail.
+		uint64 MaximumHeapSize = 0;
+
+		// The minimum alignment for resources in the heap.
+		uint32 HeapAlignment = 0;
+
+		// The latency between the completed fence value and the used fence value to invoke garbage collection of the heap.
+		uint32 GarbageCollectLatency = 0;
+
+		// Size of the texture cache before elements are evicted.
+		uint32 TextureCacheSize = kDefaultResourceCacheSize;
+
+		// Size of the buffer cache before elements are evicted.
+		uint32 BufferCacheSize = kDefaultResourceCacheSize;
+
+		// Whether all heaps should be created with the AllowAll heap flag.
+		bool bSupportsAllHeapFlags = true;
+	};
+
+	FRHITransientHeapCache(const FInitializer& InInitializer)
 		: Initializer(InInitializer)
 	{}
 
-	FRHITransientResourceSystem(const FRHITransientResourceSystem&) = delete;
+	virtual ~FRHITransientHeapCache();
+
+	FRHITransientHeap* Acquire(uint64 FirstAllocationSize, ERHITransientHeapFlags FirstAllocationHeapFlags);
 
-	virtual ~FRHITransientResourceSystem();
+	void Forfeit(TConstArrayView<FRHITransientHeap*> Heaps);
 
-	void ReleaseHeaps();
+	void GarbageCollect() override;
 
-	void GarbageCollect();
+	const FInitializer& GetInitializer() const { return Initializer; }
 
-	void UpdateStats();
+	uint64 GetGarbageCollectCycle() const { return GarbageCollectCycle; }
 
-	FORCEINLINE uint64 GetHeapSize(uint64 RequestedHeapSize) const
+	uint64 GetHeapSize(uint64 RequestedHeapSize) const
 	{
 		check(RequestedHeapSize <= Initializer.MaximumHeapSize);
 		return FMath::Clamp(FMath::RoundUpToPowerOfTwo64(RequestedHeapSize), Initializer.MinimumHeapSize, Initializer.MaximumHeapSize);
 	}
 
-	FORCEINLINE uint64 GetMinimumHeapSize() const { return Initializer.MinimumHeapSize; }
-	FORCEINLINE uint64 GetMaximumHeapSize() const { return Initializer.MaximumHeapSize; }
-	FORCEINLINE uint32 GetHeapAlignment(uint32 FirstAllocationAlignment) const { return FMath::Max(Initializer.HeapAlignment, FirstAllocationAlignment); }
-	FORCEINLINE uint32 GetHeapCount() const { return Heaps.Num(); }
+private:
+	//////////////////////////////////////////////////////////////////////////
+	//! Platform API
+
+	// Called when a new heap is being created and added to the pool.
+	virtual FRHITransientHeap* CreateHeap(const FRHITransientHeap::FInitializer& Initializer) = 0;
 
-	FORCEINLINE FRHITransientHeap* GetHeap(uint32 Index) { return Heaps[Index]; }
+	//////////////////////////////////////////////////////////////////////////
 
-protected:
-	struct FStats
-	{
-		FRHITransientMemoryStats Textures;
-		FRHITransientMemoryStats Buffers;
+	FInitializer Initializer;
+	FRHITransientMemoryStats Stats;
 
-		// Total memory used by the underlying heaps.
-		uint64 TotalMemoryUsed{};
-	};
+	FCriticalSection CriticalSection;
+	TArray<FRHITransientHeap*> LiveList;
+	TArray<FRHITransientHeap*> FreeList;
+	uint64 GarbageCollectCycle = 0;
+	uint64 TotalMemoryCapacity = 0;
 
-private:
-	// Called by the transient allocator to acquire a heap from the cache.
-	FRHITransientHeap* AcquireHeap(uint64 FirstAllocationSize, uint32 FirstAllocationAlignment, ERHITransientHeapFlags FirstAllocationHeapFlags);
+	friend FRHITransientResourceHeapAllocator;
+};
 
-	// Called by the transient allocator to forfeit all acquired heaps back to the cache.
-	void ForfeitHeaps(TConstArrayView<FRHITransientHeapState> Heaps);
+/** A base class for implementing IRHITransientResourceAllocator for a virtual aliasing placed resource heap allocation strategy. */
+class RHICORE_API FRHITransientResourceHeapAllocator : public IRHITransientResourceAllocator
+{
+public:
+	FRHITransientResourceHeapAllocator(FRHITransientHeapCache& InHeapCache)
+		: HeapCache(InHeapCache)
+	{}
 
-	//////////////////////////////////////////////////////////////////////////
-	//! Platform API
+	~FRHITransientResourceHeapAllocator();
 
-	// Called when a new heap is being created and added to the pool.
-	virtual FRHITransientHeap* CreateHeap(const FRHITransientHeapInitializer& Initializer) = 0;
+	// Deallocates a texture from its parent heap. Provide the current platform fence value used to update the heap.
+	void DeallocateMemory(FRHITransientTexture* Texture, uint32 PassIndex) override;
 
-	// Called for the platform to report stats.
-	virtual void ReportStats(const FStats& Stats) {};
+	// Deallocates a buffer from its parent heap. Provide the current platform fence value used to update the heap.
+	void DeallocateMemory(FRHITransientBuffer* Buffer, uint32 PassIndex) override;
 
-	//////////////////////////////////////////////////////////////////////////
+	// Called to flush any active allocations prior to rendering.
+	void Flush(FRHICommandListImmediate& RHICmdList, FRHITransientAllocationStats* OutAllocationStats) override;
 
-	FRHITransientResourceSystemInitializer Initializer;
+	// Returns the array of heaps used by this allocator, including the required commit size for each.
+	inline TConstArrayView<FRHITransientHeap*> GetHeaps() const { return Heaps; }
 
-	uint64 GarbageCollectCycle{};
+	FRHITransientHeapCache& HeapCache;
 
-	FCriticalSection HeapCriticalSection;
-	TArray<FRHITransientHeap*> Heaps;
+protected:
+	/** Allocates a texture on a heap at a specific offset, returning a cached RHI transient texture pointer, or null
+	 *  if the allocation failed. TextureSize and TextureAlignment are platform specific and must be derived from the
+	 *  texture create info and passed in, along with a platform-specific texture creation function if no cached resource
+	 *  if found.
+	 */
+	FRHITransientTexture* CreateTextureInternal(
+		const FRHITextureCreateInfo& CreateInfo,
+		const TCHAR* DebugName,
+		uint32 PassIndex,
+		uint64 TextureSize,
+		uint32 TextureAlignment,
+		FRHITransientHeap::FCreateTextureFunction CreateTextureFunction);
+
+	/** Allocates a buffer on a heap at a specific offset, returning a cached RHI transient buffer pointer, or null
+	 *  if the allocation failed. BufferSize and BufferAlignment are platform specific and must be derived from the
+	 *  buffer create info and passed in, along with a platform-specific buffer creation function if no cached resource
+	 *  if found.
+	 */
+	FRHITransientBuffer* CreateBufferInternal(
+		const FRHIBufferCreateInfo& CreateInfo,
+		const TCHAR* DebugName,
+		uint32 PassIndex,
+		uint32 BufferSize,
+		uint32 BufferAlignment,
+		FRHITransientHeap::FCreateBufferFunction CreateBufferFunction);
 
-	FRHITransientMemoryStats TextureStats;
-	FRHITransientMemoryStats BufferStats;
+private:
+	TArray<FRHITransientHeap*> Heaps;
+	uint64 CurrentCycle = 0;
+	uint32 DeallocationCount = 0;
 
-	friend FRHITransientResourceAllocator;
+	IF_RHICORE_TRANSIENT_ALLOCATOR_DEBUG(TSet<FRHITransientResource*> ActiveResources);
 };
 
-/** Tracks resource allocations on the heap and adds overlap events to transient resources. */
-class RHICORE_API FRHITransientResourceOverlapTracker final
+///////////////////////////////////////////////////////////////////////////////////////////////////
+//! Transient Resource Page Allocator
+///////////////////////////////////////////////////////////////////////////////////////////////////
+
+class FRHITransientPagePoolCache;
+class FRHITransientResourcePageAllocator;
+
+/** Allocates page spans for a resource. */
+class RHICORE_API FRHITransientPageSpanAllocator
 {
 public:
-	void Track(FRHITransientResource* InResource, uint64 OffsetMin, uint64 OffsetMax);
+	FRHITransientPageSpanAllocator(uint32 InPageCount, uint32 InPageSize)
+		: MaxSpanCount(InPageCount + 2)
+		, MaxPageCount(InPageCount)
+		, PageSize(InPageSize)
+	{
+		Init();
+	}
+
+	void Reset();
+
+	bool Allocate(uint32 PageCount, uint32& NumPagesAllocated, uint32& SpanIndex);
+
+	void Deallocate(uint32 SpanIndex);
+
+	template <typename SpanArrayType>
+	void GetSpanArray(uint32 SpanIndex, SpanArrayType& OutPageSpans) const
+	{
+		check(SpanIndex != InvalidIndex);
+		do
+		{
+			OutPageSpans.Emplace(PageSpans[SpanIndex]);
+			SpanIndex = PageSpans[SpanIndex].NextSpanIndex;
+		} while (SpanIndex != InvalidIndex);
+	}
+
+	uint32 GetAllocationCount() const { return AllocationCount; }
+
+	uint32 GetFreePageCount() const { return FreePageCount; }
+
+	uint64 GetUsedSize() const { return (MaxPageCount - FreePageCount) * PageSize; }
+
+	uint64 GetFreeSize() const { return FreePageCount * PageSize; }
+
+	uint32 GetAllocationPageCount(uint32 SpanIndex) const;
+
+	uint32 GetMaxSpanCount() const { return MaxSpanCount; }
+
+	uint32 GetPageSize() const { return PageSize; }
+
+	uint64 GetCapacity() const { return MaxPageCount * PageSize; }
+
+	bool IsFull() const { return FreePageCount == 0; }
+
+	bool IsEmpty() const { return AllocationCount == 0; }
 
 private:
-	struct FResourceRange
+	static const uint32 InvalidIndex = TNumericLimits<uint32>::Max();
+	static const uint32 FreeSpanListHeadIndex = 0;
+	static const uint32 FreeSpanListTailIndex = 1;
+
+	struct FPageSpan : FRHITransientPageSpan
 	{
-		FRHITransientResource* Resource{};
-		uint64 OffsetMin{};
-		uint64 OffsetMax{};
+		const bool IsLinked() { return (NextSpanIndex != InvalidIndex || PrevSpanIndex != InvalidIndex); }
+
+		uint32 NextSpanIndex = InvalidIndex;
+		uint32 PrevSpanIndex = 0;
+		bool bAllocated = false;
 	};
 
-	TArray<FResourceRange, TMemStackAllocator<>> ResourceRanges;
-};
+	void Init();
 
-/** Represents an allocation from the transient heap. */
-struct FRHITransientHeapAllocation
-{
-	FRHITransientHeapAllocation() = default;
+	// Splits a span into two, so that the original span has PageCount pages and the new span contains the remaining ones
+	void SplitSpan(uint32 SpanIndex, uint32 PageCount);
+
+	// Merges two spans. They must be adjacent and in the same list
+	void MergeSpans(uint32 SpanIndex0, uint32 SpanIndex1, const bool bKeepSpan1);
+
+	// Inserts a span after an existing span. The span to insert must be unlinked
+	void InsertAfter(uint32 InsertPosition, uint32 InsertSpanIndex);
 
-	FORCEINLINE bool IsValid() const { return Size != 0; }
+	// Inserts a span after an existing span. The span to insert must be unlinked
+	void InsertBefore(uint32 InsertPosition, uint32 InsertSpanIndex);
 
-	bool operator==(const FRHITransientHeapAllocation& Other) const
+	// Removes a span from its list, reconnecting neighbouring list elements
+	void Unlink(uint32 SpanIndex);
+
+	// Allocates an unused span from the pool
+	int AllocSpan()
 	{
-		return Size == Other.Size && Offset == Other.Offset && AlignmentPad == Other.AlignmentPad && HeapIndex == Other.HeapIndex;
+		check(UnusedSpanListCount > 0);
+		uint32 SpanIndex = UnusedSpanList[UnusedSpanListCount - 1];
+		UnusedSpanListCount--;
+		return SpanIndex;
 	}
 
-	bool operator!=(const FRHITransientHeapAllocation& Other) const
+	// Releases a span back to the unused pool
+	void ReleaseSpan(uint32 SpanIndex)
 	{
-		return !(*this == Other);
+		check(!PageSpans[SpanIndex].IsLinked());
+		UnusedSpanList[UnusedSpanListCount] = SpanIndex;
+		UnusedSpanListCount++;
+		check(UnusedSpanListCount <= MaxPageCount);
 	}
 
-	// Size of the allocation made from the allocator (aligned).
-	uint64 Size = 0;
+	// Merges a span with existing neighbours in the free list if they exist
+	bool MergeFreeSpanIfPossible(uint32 SpanIndex);
 
-	// Offset in the transient heap; front of the heap starts at 0.
-	uint64 Offset = 0;
+	void Validate();
 
-	// Number of bytes of padding were added to the offset.
-	uint32 AlignmentPad = 0;
+	int GetFirstSpanIndex() const
+	{
+		return PageSpans[FreeSpanListHeadIndex].NextSpanIndex;
+	}
 
-	// Index of the transient heap.
-	uint32 HeapIndex = 0;
+	TArray<int32> PageToSpanStart;  // [PAGE_COUNT + 1]
+	TArray<int32> PageToSpanEnd;    // [PAGE_COUNT + 1]
+	TArray<FPageSpan> PageSpans;    // [MAX_SPAN_COUNT]
+	TArray<int32> UnusedSpanList;   // [MAX_SPAN_COUNT]
+
+	uint32 FreePageCount;
+	uint32 UnusedSpanListCount;
+
+	const uint32 MaxSpanCount;
+	const uint32 MaxPageCount;
+	const uint32 PageSize;
+	uint32 AllocationCount;
 };
 
-/** A simple first-fit allocator for placing resources onto a transient heap and tracking their aliasing overlap events. */
-class RHICORE_API FRHITransientHeapAllocator final
+///////////////////////////////////////////////////////////////////////////////////////////////////
+
+class RHICORE_API FRHITransientPagePool
 {
 public:
-	FRHITransientHeapAllocator(const FRHITransientHeapInitializer& Initializer, uint32 HeapIndex, uint64 InHeapBaseGPUAddress = 0);
-	~FRHITransientHeapAllocator();
+	struct FInitializer
+	{
+		uint32 PageCount = 0;
+		uint32 PageSize = 0;
+	};
 
-	FRHITransientHeapAllocation Allocate(uint64 Size, uint32 Alignment);
-	void Deallocate(FRHITransientHeapAllocation Allocation);
+	FRHITransientPagePool(const FInitializer& InInitializer)
+		: Initializer(InInitializer)
+		, Allocator(Initializer.PageCount, Initializer.PageSize)
+	{}
 
-	FORCEINLINE void TrackOverlap(FRHITransientResource* Resource, const FRHITransientHeapAllocation& Allocation)
+	virtual ~FRHITransientPagePool() = default;
+
+	struct FAllocationContext
 	{
-		OverlapTracker.Track(Resource, Allocation.Offset, Allocation.Offset + Allocation.Size);
-	}
+		FAllocationContext(FRHITransientResource& InResource, uint32 InPageSize)
+			: Resource(InResource)
+			, Allocations(InResource.GetPageAllocation().PoolAllocations)
+			, Spans(InResource.GetPageAllocation().Spans)
+			, AllocationsBefore(Allocations)
+			, GpuVirtualAddress(Resource.GetGpuVirtualAddress())
+			, Size(Resource.GetSize())
+			, PagesRemaining(Size / InPageSize)
+		{
+			checkf(Size % InPageSize == 0, TEXT("Resource size was not aligned properly to page size."));
+			Allocations.Reset();
+			Spans.Reset();
+		}
 
-	FORCEINLINE uint64 GetCapacity() const { return Initializer.Size; }
-	FORCEINLINE uint64 GetUsedSize() const { return UsedSize; }
-	FORCEINLINE uint64 GetFreeSize() const { return Initializer.Size - UsedSize; }
-	FORCEINLINE uint64 GetAlignmentWaste() const { return AlignmentWaste; }
-	FORCEINLINE uint32 GetAllocationCount() const { return AllocationCount; }
+		bool IsComplete() const { return PagesRemaining == 0; }
+
+		FRHITransientResource& Resource;
+		TArray<FRHITransientPagePoolAllocation>& Allocations;
+		TArray<FRHITransientPageSpan>& Spans;
+		const TArray<FRHITransientPagePoolAllocation, TInlineAllocator<8>> AllocationsBefore;
+		const uint64 GpuVirtualAddress;
+		const uint64 Size;
+
+		uint32 AllocationCount = 0;
+		uint32 AllocationMatchingCount = 0;
+		uint32 PagesRemaining = 0;
+		uint32 PagesAllocated = 0;
+		uint32 PageSpansAllocated = 0;
+		uint32 PagesMapped = 0;
+	};
 
-	FORCEINLINE bool IsFull() const { return UsedSize == Initializer.Size; }
-	FORCEINLINE bool IsEmpty() const { return UsedSize == 0; }
+	void Allocate(FAllocationContext& AllocationContext);
 
-	// Returns whether the requested allocation can succeed.
-	FORCEINLINE bool IsAllocationSupported(uint64 Size, ERHITransientHeapFlags Flags) const
+	void Deallocate(uint32 SpanIndex)
 	{
-		return Size <= GetFreeSize() && EnumHasAnyFlags(Initializer.Flags, Flags);
+		Allocator.Deallocate(SpanIndex);
 	}
 
-private:
-	using FRangeHandle = uint16;
-	static const FRangeHandle InvalidRangeHandle = FRangeHandle(~0);
+	void Flush(FRHICommandListImmediate& RHICmdList);
 
-	struct FRange
+	uint64 GetLastUsedGarbageCollectCycle() const { return LastUsedGarbageCollectCycle; }
+
+	bool IsEmpty() const { return Allocator.IsEmpty(); }
+
+	bool IsFull() const { return Allocator.IsFull(); }
+
+	uint64 GetCapacity() const { return Allocator.GetCapacity(); }
+
+	uint64 GetGpuVirtualAddress() const { return GpuVirtualAddress; }
+
+	struct FPageMapRequest
 	{
-		uint64 Size{};
-		uint64 Offset{};
-		FRangeHandle NextFreeHandle = InvalidRangeHandle;
+		FPageMapRequest() = default;
+
+		FPageMapRequest(
+			uint64 InDestinationAddress,
+			uint64 InSourcePagePoolAddress,
+			uint32 InSourcePageCount,
+			uint32 InPageSpanOffset,
+			uint32 InPageSpanCount)
+			: DestinationAddress(InDestinationAddress)
+			, SourcePagePoolAddress(InSourcePagePoolAddress)
+			, SourcePageCount(InSourcePageCount)
+			, PageSpanOffset(InPageSpanOffset)
+			, PageSpanCount(InPageSpanCount)
+		{}
 
-		FORCEINLINE uint64 GetStart() const { return Offset; }
-		FORCEINLINE uint64 GetEnd() const { return Size + Offset; }
+		uint64 DestinationAddress = 0;
+		uint64 SourcePagePoolAddress = 0;
+		uint32 SourcePageCount = 0;
+		uint32 PageSpanOffset = 0;
+		uint32 PageSpanCount = 0;
 	};
 
-	FORCEINLINE FRangeHandle GetFirstFreeRangeHandle()
-	{
-		return Ranges[HeadHandle].NextFreeHandle;
-	}
+	const FInitializer Initializer;
 
-	FRangeHandle CreateRange()
+protected:
+	void SetGpuVirtualAddress(uint64 InGpuVirtualAddress)
 	{
-		if (!RangeFreeList.IsEmpty())
-		{
-			return RangeFreeList.Pop();
-		}
-		Ranges.Emplace();
-		return FRangeHandle(Ranges.Num() - 1);
+		check(InGpuVirtualAddress != 0);
+		GpuVirtualAddress = InGpuVirtualAddress;
 	}
 
-	FRangeHandle InsertRange(FRangeHandle PreviousHandle, uint64 Offset, uint64 Size)
-	{
-		FRangeHandle Handle = CreateRange();
+private:
+	//////////////////////////////////////////////////////////////////////////
+	//! Platform API
 
-		FRange& CurrentRange = Ranges[Handle];
-		CurrentRange.Offset = Offset;
-		CurrentRange.Size = Size;
+	virtual void Flush(FRHICommandListImmediate& RHICmdList, TArray<FPageMapRequest>&& InPageMapRequests, TArray<FRHITransientPageSpan>&& InPageSpans) = 0;
 
-		FRange& PreviousRange = Ranges[PreviousHandle];
-		CurrentRange.NextFreeHandle = PreviousRange.NextFreeHandle;
-		PreviousRange.NextFreeHandle = Handle;
+	//////////////////////////////////////////////////////////////////////////
 
-		return Handle;
-	}
+	FRHITransientPageSpanAllocator Allocator;
+	FRHITransientResourceOverlapTracker OverlapTracker;
 
-	void RemoveRange(FRangeHandle PreviousHandle, FRangeHandle CurrentHandle)
-	{
-		FRange& PreviousRange = Ranges[PreviousHandle];
-		FRange& CurrentRange = Ranges[CurrentHandle];
+	TArray<FPageMapRequest> PageMapRequests;
+	TArray<FRHITransientPageSpan> PageSpans;
+	uint64 GpuVirtualAddress = 0;
+	uint64 LastUsedGarbageCollectCycle = 0;
+	uint32 PageMapRequestCountMax = 0;
+	uint32 PageSpanCountMax = 0;
 
-		PreviousRange.NextFreeHandle = CurrentRange.NextFreeHandle;
-		CurrentRange.NextFreeHandle = InvalidRangeHandle;
+	friend FRHITransientPagePoolCache;
+	friend FRHITransientResourcePageAllocator;
+};
 
-		RangeFreeList.Add(CurrentHandle);
-	}
+///////////////////////////////////////////////////////////////////////////////////////////////////
 
-	struct FFindResult
+class RHICORE_API FRHITransientPagePoolCache : public IRHITransientMemoryCache
+{
+public:
+	struct FInitializer
 	{
-		uint64 LeftoverSize = 0;
-		FRangeHandle PreviousHandle = InvalidRangeHandle;
-		FRangeHandle FoundHandle = InvalidRangeHandle;
+		// Creates a default initializer using common RHI CVars.
+		static FInitializer CreateDefault();
+
+		static const uint32 kDefaultResourceCacheSize = 256;
+
+		// Size in bytes of the pool. Must be a multiple of PageSize.
+		uint64 PoolSize = 0;
+
+		// Size in bytes of the first pool. Only takes effect if larger than PoolSize.
+		uint64 PoolSizeFirst = 0;
+
+		// Size of each page.
+		uint32 PageSize = 0;
+
+		// The latency between the completed fence value and the used fence value to invoke garbage collection of the heap.
+		uint32 GarbageCollectLatency = 0;
+
+		// Size of the texture cache before elements are evicted.
+		uint32 TextureCacheSize = kDefaultResourceCacheSize;
+
+		// Size of the buffer cache before elements are evicted.
+		uint32 BufferCacheSize = kDefaultResourceCacheSize;
 	};
 
-	FFindResult FindFreeRange(uint64 Size, uint32 Alignment);
+	FRHITransientPagePoolCache(const FInitializer& InInitializer)
+		: Initializer(InInitializer)
+	{}
 
-	void Validate();
+	FRHITransientPagePoolCache(const FRHITransientPagePoolCache&) = delete;
 
-	FRHITransientHeapInitializer Initializer;
+	virtual ~FRHITransientPagePoolCache();
 
-	uint64 UsedSize{};
-	uint64 AlignmentWaste{};
-	uint32 AllocationCount{};
-	uint32 HeapIndex{};
+	// Called by the transient allocator to acquire a page pool from the cache.
+	FRHITransientPagePool* Acquire();
 
-	uint64 HeapBaseGPUAddress;
+	// Called by the transient allocator to forfeit all acquired heaps back to the cache.
+	void Forfeit(TConstArrayView<FRHITransientPagePool*> PagePools);
 
-	FRangeHandle HeadHandle = InvalidRangeHandle;
-	TArray<FRangeHandle, TInlineAllocator<4, TMemStackAllocator<>>> RangeFreeList;
-	TArray<FRange, TInlineAllocator<4, TMemStackAllocator<>>> Ranges;
+	void GarbageCollect() override;
 
-	FRHITransientResourceOverlapTracker OverlapTracker;
-};
+	const FInitializer Initializer;
 
-/** A helper class for implementing IRHITransientResourceAllocator. This class is designed for composition instead of
- *  inheritance to keep the platform implementation clean. Its methods are not straight overrides of the user-facing
- *  class; rather, they are designed to streamline the platform implementation. This object is designed to match the
- *  lifecycle of IRHITransientResourceAllocator. Specifically, it must be short-lived and its API used on the render
- *  thread, as it utilizes MemStack-allocated containers for performance. Its state is built from scratch each allocation
- *  cycle and released at the end. Heaps are acquired from the parent transient resource system, and transient resources
- *  from respective heaps.
- */
-class RHICORE_API FRHITransientResourceAllocator final
-{
-public:
-	FRHITransientResourceAllocator(FRHITransientResourceSystem& InParentSystem);
-	FRHITransientResourceAllocator(const FRHITransientResourceAllocator&) = delete;
+private:
+	//////////////////////////////////////////////////////////////////////////
+	//! Platform API
 
-	struct FResourceInitializer
-	{
-		FResourceInitializer(FRHITransientHeap& InHeap, const FRHITransientHeapAllocation& InAllocation, uint64 InHash)
-			: Heap(InHeap)
-			, Allocation(InAllocation)
-			, Hash(InHash)
-		{}
+	// Called to access the dedicated fast VRAM page pool, if it exists on the platform.
+	virtual FRHITransientPagePool* GetFastPagePool() = 0;
 
-		// The heap on which to create the resource.
-		FRHITransientHeap& Heap;
+	// Called when a new heap is being created and added to the pool.
+	virtual FRHITransientPagePool* CreatePagePool(const FRHITransientPagePool::FInitializer& Initializer) = 0;
 
-		// The allocation (offset / size) on the provided heap.
-		const FRHITransientHeapAllocation& Allocation;
+	//////////////////////////////////////////////////////////////////////////
 
-		// The unique hash computed from the create info and allocation offset.
-		const uint64 Hash;
-	};
+	FRHITransientMemoryStats Stats;
 
-	using FCreateTextureFunction = TFunction<FRHITransientTexture*(const FResourceInitializer&)>;
+	FCriticalSection CriticalSection;
+	TArray<FRHITransientPagePool*> LiveList;
+	TArray<FRHITransientPagePool*> FreeList;
+	uint64 GarbageCollectCycle = 0;
+	uint64 TotalMemoryCapacity = 0;
 
-	/** Allocates a texture on a heap at a specific offset, returning a cached RHI transient texture pointer, or null
-	 *  if the allocation failed. TextureSize and TextureAlignment are platform specific and must be derived from the
-	 *  texture create info and passed in, along with a platform-specific texture creation function if no cached resource
-	 *  if found.
-	 */
-	FRHITransientTexture* CreateTexture(
-		const FRHITextureCreateInfo& CreateInfo,
-		const TCHAR* DebugName,
-		uint32 PassIndex,
-		uint64 TextureSize,
-		uint32 TextureAlignment,
-		FCreateTextureFunction CreateTextureFunction);
+	friend FRHITransientResourcePageAllocator;
+};
 
-	using FCreateBufferFunction = TFunction<FRHITransientBuffer*(const FResourceInitializer&)>;
+///////////////////////////////////////////////////////////////////////////////////////////////////
 
-	/** Allocates a buffer on a heap at a specific offset, returning a cached RHI transient buffer pointer, or null
-	 *  if the allocation failed. BufferSize and BufferAlignment are platform specific and must be derived from the
-	 *  buffer create info and passed in, along with a platform-specific buffer creation function if no cached resource
-	 *  if found.
-	 */
-	FRHITransientBuffer* CreateBuffer(
-		const FRHIBufferCreateInfo& CreateInfo,
-		const TCHAR* DebugName,
-		uint32 PassIndex,
-		uint32 BufferSize,
-		uint32 BufferAlignment,
-		FCreateBufferFunction CreateBufferFunction);
+class RHICORE_API FRHITransientResourcePageAllocator : public IRHITransientResourceAllocator
+{
+public:
+	FRHITransientResourcePageAllocator(FRHITransientPagePoolCache& InPagePoolCache)
+		: PagePoolCache(InPagePoolCache)
+		, Textures(InPagePoolCache.Initializer.TextureCacheSize, ResourceCacheGarbageCollectionLatency)
+		, Buffers(InPagePoolCache.Initializer.BufferCacheSize, ResourceCacheGarbageCollectionLatency)
+		, PageSize(PagePoolCache.Initializer.PageSize)
+	{}
 
-	// Deallocates a texture from its parent heap. Provide the current platform fence value used to update the heap.
-	FORCEINLINE void DeallocateMemory(FRHITransientTexture* Texture, uint32 PassIndex)
-	{
-		DeallocateMemoryInternal(Texture, PassIndex, TextureStats);
-	}
+	~FRHITransientResourcePageAllocator();
 
-	// Deallocates a buffer from its parent heap. Provide the current platform fence value used to update the heap.
-	FORCEINLINE void DeallocateMemory(FRHITransientBuffer* Buffer, uint32 PassIndex)
-	{
-		DeallocateMemoryInternal(Buffer, PassIndex, BufferStats);
-	}
+	FRHITransientTexture* CreateTexture(const FRHITextureCreateInfo& CreateInfo, const TCHAR* DebugName, uint32 PassIndex) override;
+	FRHITransientBuffer* CreateBuffer(const FRHIBufferCreateInfo& CreateInfo, const TCHAR* DebugName, uint32 PassIndex) override;
+	void DeallocateMemory(FRHITransientTexture* Texture, uint32 PassIndex) override;
+	void DeallocateMemory(FRHITransientBuffer* Buffer, uint32 PassIndex) override;
+	void Flush(FRHICommandListImmediate& RHICmdList, FRHITransientAllocationStats* OutAllocationStats) override;
 
-	// Called to signify all allocations have completed. Forfeits all resources / heaps back to the parent system.
-	void Freeze(FRHICommandListImmediate& RHICmdList, FRHITransientHeapStats& OutHeapStats);
+	uint32 GetPageSize() const { return PageSize; }
 
-	// Called to flush any active allocations prior to Freeze.
-	void Flush();
+	uint32 GetPagePoolCount() const { return PagePools.Num(); }
 
-	// Returns the array of heaps used by this allocator, including the required commit size for each.
-	inline TConstArrayView<FRHITransientHeapState> GetUsedHeapStates() const
-	{
-		return Heaps;
-	}
+	FRHITransientPagePoolCache& PagePoolCache;
 
 private:
-	struct FMemoryStats : FRHITransientMemoryStats
-	{
-		uint64 CurrentSizeWithAliasing{};
-	};
+	static constexpr uint32 ResourceCacheGarbageCollectionLatency = 2;
+	static constexpr uint64 KB = 1024;
+	static constexpr uint64 MB = 1024 * KB;
 
-	FRHITransientHeapAllocation Allocate(FMemoryStats& StatsToUpdate, uint64 Size, uint32 Alignment, ERHITransientHeapFlags ResourceHeapFlags);
+	void AllocateMemoryInternal(FRHITransientResource* Resource, const TCHAR* DebugName, uint32 PassIndex, bool bFastPoolRequested);
+	void DeallocateMemoryInternal(FRHITransientResource* Resource, uint32 PassIndex);
 
-	void DeallocateMemoryInternal(FRHITransientResource* InResource, uint32 PassIndex, FMemoryStats& StatsToUpdate);
+	//////////////////////////////////////////////////////////////////////////
+	//! Platform API
 
-	void InitResource(FRHITransientResource* TransientResource, uint32 PassIndex, const FRHITransientHeapAllocation& Allocation, const TCHAR* Name);
+	virtual FRHITransientTexture* CreateTextureInternal(
+		const FRHITextureCreateInfo& CreateInfo,
+		const TCHAR* DebugName,
+		uint64 Hash) = 0;
 
-	FRHITransientResourceSystem& ParentSystem;
+	virtual FRHITransientBuffer* CreateBufferInternal(
+		const FRHIBufferCreateInfo& CreateInfo,
+		const TCHAR* DebugName,
+		uint64 Hash) = 0;
 
-	// Tracks state on the rendering thread; must be cleared before the destructor.
-	TArray<FRHITransientHeapState, TInlineAllocator<4, TMemStackAllocator<>>> Heaps;
-	TArray<FRHITransientHeapAllocator, TInlineAllocator<4, TMemStackAllocator<>>> HeapAllocators;
-	TArray<FRHITransientHeapAllocation, TInlineAllocator<4, TMemStackAllocator<>>> HeapAllocations;
+	virtual void ReleaseTextureInternal(FRHITransientTexture* Texture) = 0;
+	virtual void ReleaseBufferInternal(FRHITransientBuffer* Buffer) = 0;
 
-#if RHICORE_TRANSIENT_ALLOCATOR_DEBUG
-	TArray<FRHITransientTexture*, TMemStackAllocator<>> DebugTextures;
-	TArray<FRHITransientBuffer*, TMemStackAllocator<>> DebugBuffers;
-#endif
+	//////////////////////////////////////////////////////////////////////////
+
+	template <typename FunctionType>
+	void EnumeratePageSpans(const FRHITransientResource* Resource, FunctionType Function) const
+	{
+		const FRHITransientPageAllocation& PageAllocation = Resource->GetPageAllocation();
+
+		for (const FRHITransientPagePoolAllocation& PoolAllocation : PageAllocation.PoolAllocations)
+		{
+			for (uint32 Index = PoolAllocation.SpanOffsetMin; Index < PoolAllocation.SpanOffsetMax; ++Index)
+			{
+				Function(PoolAllocation.Pool, PageAllocation.Spans[Index]);
+			}
+		}
+	}
+
+	FRHITransientMemoryStats Stats;
+	TRHITransientResourceCache<FRHITransientTexture> Textures;
+	TRHITransientResourceCache<FRHITransientBuffer> Buffers;
+
+	TArray<FRHITransientPagePool*> PagePools;
+	uint64 CurrentCycle = 0;
+	uint32 DeallocationCount = 0;
+	uint32 PageSize = 0;
+	uint32 PageMapCount = 0;
+	uint32 PageAllocateCount = 0;
+	uint32 PageSpanCount = 0;
 
-	FMemoryStats TextureStats;
-	FMemoryStats BufferStats;
+	IF_RHICORE_TRANSIENT_ALLOCATOR_DEBUG(TSet<FRHITransientResource*> ActiveResources);
 };
\ No newline at end of file
diff --git a/Engine/Source/Runtime/RenderCore/Private/RenderGraphBuilder.cpp b/Engine/Source/Runtime/RenderCore/Private/RenderGraphBuilder.cpp
--- a/Engine/Source/Runtime/RenderCore/Private/RenderGraphBuilder.cpp
+++ b/Engine/Source/Runtime/RenderCore/Private/RenderGraphBuilder.cpp
@@ -489,12 +489,24 @@ bool FRDGBuilder::IsTransientInternal(FRDGParentResourceRef Resource) const
 		return false;
 	}
 
-	// Transient resources must stay within the graph.
-	if (Resource->bExternal || Resource->bExtracted || Resource->bUserSetNonTransient)
+	if (Resource->bForceNonTransient)
 	{
 		return false;
 	}
 
+	if (Resource->bExtracted)
+	{
+		if (GRDGTransientExtractedResources == 0)
+		{
+			return false;
+		}
+
+		if (!Resource->bAllowTransientExtracted)
+		{
+			return false;
+		}
+	}
+
 #if RDG_ENABLE_DEBUG
 	if (GRDGDebugDisableTransientResources != 0 && IsDebugAllowedForResource(Resource->Name))
 	{
@@ -505,25 +517,6 @@ bool FRDGBuilder::IsTransientInternal(FRDGParentResourceRef Resource) const
 	return true;
 }
 
-FRDGBuilder::FTransientResourceAllocator::~FTransientResourceAllocator()
-{
-	if (Allocator)
-	{
-		Allocator->Release(RHICmdList);
-	}
-}
-
-IRHITransientResourceAllocator* FRDGBuilder::FTransientResourceAllocator::GetOrCreate()
-{
-	check(!IsImmediateMode());
-	if (!Allocator && !bCreateAttempted)
-	{
-		Allocator = RHICreateTransientResourceAllocator();
-		bCreateAttempted = true;
-	}
-	return Allocator;
-}
-
 FRDGBuilder::FRDGBuilder(FRHICommandListImmediate& InRHICmdList, FRDGEventName InName, ERDGBuilderFlags InFlags)
 	: RHICmdList(InRHICmdList)
 	, Blackboard(Allocator)
@@ -542,7 +535,7 @@ FRDGBuilder::FRDGBuilder(FRHICommandListImmediate& InRHICmdList, FRDGEventName I
 	, BarrierValidation(&Passes, BuilderName)
 	, LogFile(Passes)
 #endif
-	, TransientResourceAllocator(RHICmdList)
+	, TransientResourceAllocator(GRDGTransientResourceAllocator.Get())
 {
 	AddProloguePass();
 
@@ -579,6 +572,7 @@ const TRefCountPtr<FRDGPooledBuffer>& FRDGBuilder::ConvertToExternalBuffer(FRDGB
 	if (!Buffer->bExternal)
 	{
 		Buffer->bExternal = 1;
+		Buffer->bForceNonTransient = 1;
 		Buffer->AccessFinal = kDefaultAccessFinal;
 		BeginResourceRHI(GetProloguePassHandle(), Buffer);
 		ExternalBuffers.Add(Buffer->PooledBuffer, Buffer);
@@ -592,6 +586,7 @@ const TRefCountPtr<IPooledRenderTarget>& FRDGBuilder::ConvertToExternalTexture(F
 	if (!Texture->bExternal)
 	{
 		Texture->bExternal = 1;
+		Texture->bForceNonTransient = 1;
 		Texture->AccessFinal = kDefaultAccessFinal;
 		BeginResourceRHI(GetProloguePassHandle(), Texture);
 		ExternalTextures.Add(Texture->GetRHIUnchecked(), Texture);
@@ -693,7 +688,7 @@ FRDGTextureRef FRDGBuilder::RegisterExternalTexture(
 	}
 
 	FRDGTextureRef Texture = Textures.Allocate(Allocator, Name, Desc, Flags, RenderTargetTexture);
-	Texture->SetRHI(static_cast<FPooledRenderTarget*>(ExternalPooledTexture.GetReference()));
+	Texture->SetRHI(ExternalPooledTexture.GetReference());
 
 	Texture->bExternal = true;
 	Texture->AccessFinal = EnumHasAnyFlags(ExternalTextureRHI->GetFlags(), TexCreate_Foveation) ? ERHIAccess::ShadingRateSource : kDefaultAccessFinal;
@@ -1486,6 +1481,8 @@ void FRDGBuilder::Execute(EExecuteMode ExecuteMode)
 	CSV_SCOPED_TIMING_STAT_EXCLUSIVE(RDG);
 	SCOPED_NAMED_EVENT_TEXT("FRDGBuilder::Execute", FColor::Magenta);
 
+	GRDGTransientResourceAllocator.ReleasePendingDeallocations();
+
 	// Create the epilogue pass at the end of the graph just prior to compilation.
 	{
 		bInDebugPassScope = true;
@@ -1572,16 +1569,6 @@ void FRDGBuilder::Execute(EExecuteMode ExecuteMode)
 					}
 				}
 
-				for (const FExtractedTexture& ExtractedTexture : ExtractedTextures)
-				{
-					EndResourceRHI(EpiloguePassHandle, ExtractedTexture.Texture, 1);
-				}
-
-				for (const FExtractedBuffer& ExtractedBuffer : ExtractedBuffers)
-				{
-					EndResourceRHI(EpiloguePassHandle, ExtractedBuffer.Buffer, 1);
-				}
-
 				EnumerateExtendedLifetimeResources(Textures, [&](FRDGTextureRef Texture)
 				{
 					EndResourceRHI(EpiloguePassHandle, Texture, 1);
@@ -1594,9 +1581,11 @@ void FRDGBuilder::Execute(EExecuteMode ExecuteMode)
 
 				if (TransientResourceAllocator)
 				{
-					FRHITransientHeapStats TransientHeapStats;
-					TransientResourceAllocator->Freeze(RHICmdList, TransientHeapStats);
-					IF_RDG_ENABLE_TRACE(Trace.SetTransientHeapStats(TransientHeapStats));
+				#if RDG_ENABLE_TRACE
+					TransientResourceAllocator->Flush(RHICmdList, Trace.IsEnabled() ? &Trace.TransientAllocationStats : nullptr);
+				#else
+					TransientResourceAllocator->Flush(RHICmdList);
+				#endif
 				}
 			}
 			else
@@ -2695,7 +2684,7 @@ void FRDGBuilder::AddEpilogueTransition(FRDGTextureRef Texture)
 	FRDGSubresourceState ScratchSubresourceState;
 
 	// Texture is using the RHI transient allocator. Transition it back to Discard in the final pass it is used.
-	if (Texture->bTransient)
+	if (Texture->bTransient && !Texture->TransientTexture->IsAcquired())
 	{
 		const TInterval<uint32> DiscardPasses = Texture->TransientTexture->GetDiscardPasses();
 		const FRDGPassHandle MinDiscardPassHandle(DiscardPasses.Min);
@@ -3127,31 +3116,35 @@ void FRDGBuilder::BeginResourceRHI(FRDGPassHandle PassHandle, FRDGTextureRef Tex
 	}
 #endif
 
-	if (IsTransient(Texture))
+	if (TransientResourceAllocator && IsTransient(Texture))
 	{
-		if (IRHITransientResourceAllocator* AllocatorRHI = TransientResourceAllocator.GetOrCreate())
+		if (FRHITransientTexture* TransientTexture = TransientResourceAllocator->CreateTexture(Texture->Desc, Texture->Name, PassHandle.GetIndex()))
 		{
-			if (FRHITransientTexture* TransientTexture = AllocatorRHI->CreateTexture(Texture->Desc, Texture->Name, PassHandle.GetIndex()))
+			if (Texture->bExternal || Texture->bExtracted)
+			{
+				Texture->SetRHI(GRDGTransientResourceAllocator.AllocateRenderTarget(TransientTexture));
+			}
+			else
 			{
-				Texture->SetRHI(TransientTexture, Allocator);
+				Texture->SetRHI(TransientTexture, Allocator.AllocNoDestruct<FRDGTextureSubresourceState>());
+			}
 
-				const FRDGPassHandle MinAcquirePassHandle = ClampToPrologue(FRDGPassHandle(TransientTexture->GetAcquirePasses().Min));
+			const FRDGPassHandle MinAcquirePassHandle = ClampToPrologue(FRDGPassHandle(TransientTexture->GetAcquirePasses().Min));
 
-				AddAliasingTransition(MinAcquirePassHandle, PassHandle, Texture, FRHITransientAliasingInfo::Acquire(TransientTexture->GetRHI(), TransientTexture->GetAliasingOverlaps()));
+			AddAliasingTransition(MinAcquirePassHandle, PassHandle, Texture, FRHITransientAliasingInfo::Acquire(TransientTexture->GetRHI(), TransientTexture->GetAliasingOverlaps()));
 
-				FRDGSubresourceState InitialState;
-				InitialState.SetPass(ERHIPipeline::Graphics, MinAcquirePassHandle);
-				InitialState.Access = ERHIAccess::Discard;
-				InitAsWholeResource(Texture->GetState(), InitialState);
+			FRDGSubresourceState InitialState;
+			InitialState.SetPass(ERHIPipeline::Graphics, MinAcquirePassHandle);
+			InitialState.Access = ERHIAccess::Discard;
+			InitAsWholeResource(Texture->GetState(), InitialState);
 
-			#if STATS
-				GRDGStatTransientTextureCount++;
-			#endif
-			}
+		#if STATS
+			GRDGStatTransientTextureCount++;
+		#endif
 		}
 	}
 
-	if (!Texture->bTransient)
+	if (!Texture->ResourceRHI)
 	{
 		Texture->SetRHI(GRenderTargetPool.FindFreeElementForRDG(RHICmdList, Texture->Desc, Texture->Name));
 	}
@@ -3218,26 +3211,23 @@ void FRDGBuilder::BeginResourceRHI(FRDGPassHandle PassHandle, FRDGBufferRef Buff
 	Buffer->FinalizeDesc();
 
 	// If transient then create the resource on the transient allocator. External or extracted resource can't be transient because of lifetime tracking issues.
-	if (IsTransient(Buffer))
+	if (TransientResourceAllocator && IsTransient(Buffer))
 	{
-		if (IRHITransientResourceAllocator* AllocatorRHI = TransientResourceAllocator.GetOrCreate())
+		if (FRHITransientBuffer* TransientBuffer = TransientResourceAllocator->CreateBuffer(Translate(Buffer->Desc), Buffer->Name, PassHandle.GetIndex()))
 		{
-			if (FRHITransientBuffer* TransientBuffer = AllocatorRHI->CreateBuffer(Translate(Buffer->Desc), Buffer->Name, PassHandle.GetIndex()))
-			{
-				Buffer->SetRHI(TransientBuffer, Allocator);
+			Buffer->SetRHI(TransientBuffer, Allocator);
 
-				const FRDGPassHandle MinAcquirePassHandle = ClampToPrologue(FRDGPassHandle(TransientBuffer->GetAcquirePasses().Min));
+			const FRDGPassHandle MinAcquirePassHandle = ClampToPrologue(FRDGPassHandle(TransientBuffer->GetAcquirePasses().Min));
 
-				AddAliasingTransition(MinAcquirePassHandle, PassHandle, Buffer, FRHITransientAliasingInfo::Acquire(TransientBuffer->GetRHI(), TransientBuffer->GetAliasingOverlaps()));
+			AddAliasingTransition(MinAcquirePassHandle, PassHandle, Buffer, FRHITransientAliasingInfo::Acquire(TransientBuffer->GetRHI(), TransientBuffer->GetAliasingOverlaps()));
 
-				FRDGSubresourceState& InitialState = Buffer->GetState();
-				InitialState.SetPass(ERHIPipeline::Graphics, MinAcquirePassHandle);
-				InitialState.Access = ERHIAccess::Discard;
+			FRDGSubresourceState& InitialState = Buffer->GetState();
+			InitialState.SetPass(ERHIPipeline::Graphics, MinAcquirePassHandle);
+			InitialState.Access = ERHIAccess::Discard;
 
-			#if STATS
-				GRDGStatTransientBufferCount++;
-			#endif
-			}
+		#if STATS
+			GRDGStatTransientBufferCount++;
+		#endif
 		}
 	}
 
@@ -3336,12 +3326,30 @@ void FRDGBuilder::EndResourceRHI(FRDGPassHandle PassHandle, FRDGTextureRef Textu
 	{
 		if (Texture->bTransient)
 		{
-			TransientResourceAllocator->DeallocateMemory(Texture->TransientTexture, PassHandle.GetIndex());
+			// Texture is using a transient external render target.
+			if (Texture->PooledRenderTarget)
+			{
+				// This releases the reference without invoking a virtual function call.
+				GRDGTransientResourceAllocator.Release(TRefCountPtr<FRDGTransientRenderTarget>(MoveTemp(Texture->Allocation)), PassHandle);
+			}
+			// Texture is using an internal transient texture.
+			else
+			{
+				TransientResourceAllocator->DeallocateMemory(Texture->TransientTexture, PassHandle.GetIndex());
+			}
 		}
-		else if (Texture->PooledRenderTarget && Texture->PooledRenderTarget->IsTracked())
+		else
 		{
-			// This releases the reference without invoking a virtual function call.
-			TRefCountPtr<FPooledRenderTarget>(MoveTemp(Texture->Allocation));
+			// If this is a non-transient texture, it must be backed by a pooled render target.
+			FPooledRenderTarget* RenderTarget = static_cast<FPooledRenderTarget*>(Texture->PooledRenderTarget);
+			check(RenderTarget);
+
+			// Only tracked render targets are released. Untracked ones persist until the end of the frame.
+			if (RenderTarget->IsTracked())
+			{
+				// This releases the reference without invoking a virtual function call.
+				TRefCountPtr<FPooledRenderTarget>(MoveTemp(Texture->Allocation));
+			}
 		}
 
 		Texture->LastPass = PassHandle;
diff --git a/Engine/Source/Runtime/RenderCore/Private/RenderGraphPrivate.cpp b/Engine/Source/Runtime/RenderCore/Private/RenderGraphPrivate.cpp
--- a/Engine/Source/Runtime/RenderCore/Private/RenderGraphPrivate.cpp
+++ b/Engine/Source/Runtime/RenderCore/Private/RenderGraphPrivate.cpp
@@ -319,6 +319,14 @@ FAutoConsoleVariableRef CVarRDGUseTransientAllocator(
 	TEXT(" 2: enables the transient allocator for resources with FastVRAM flag only"),
 	ECVF_RenderThreadSafe);
 
+int32 GRDGTransientExtractedResources = 1;
+FAutoConsoleVariableRef CVarRDGTransientExtractedResource(
+	TEXT("r.RDG.TransientExtractedResources"), GRDGTransientExtractedResources,
+	TEXT("RDG will allocate extracted resources as transient, unless explicitly marked non-transient by the user.")
+	TEXT(" 0: disables external transient resources;")
+	TEXT(" 1: enables external transient resources (default);"),
+	ECVF_RenderThreadSafe);
+
 int32 GRDGParallelExecute = 1;
 FAutoConsoleVariableRef CVarRDGParallelExecute(
 	TEXT("r.RDG.ParallelExecute"), GRDGParallelExecute,
diff --git a/Engine/Source/Runtime/RenderCore/Private/RenderGraphPrivate.h b/Engine/Source/Runtime/RenderCore/Private/RenderGraphPrivate.h
--- a/Engine/Source/Runtime/RenderCore/Private/RenderGraphPrivate.h
+++ b/Engine/Source/Runtime/RenderCore/Private/RenderGraphPrivate.h
@@ -107,6 +107,7 @@ extern int32 GRDGAsyncCompute;
 extern int32 GRDGCullPasses;
 extern int32 GRDGMergeRenderPasses;
 extern int32 GRDGTransientAllocator;
+extern int32 GRDGTransientExtractedResources;
 extern int32 GRDGTransientIndirectArgBuffers;
 extern int32 GRDGDrain;
 extern int32 GRDGParallelExecute;
@@ -222,7 +223,7 @@ inline void EnumerateExtendedLifetimeResources(ResourceRegistryType& Registry, F
 		{
 			auto* Resource = Registry[Handle];
 
-			if (IsDebugAllowedForResource(Resource->Name))
+			if (IsDebugAllowedForResource(Resource->Name) && !Resource->IsCulled())
 			{
 				Function(Resource);
 			}
diff --git a/Engine/Source/Runtime/RenderCore/Private/RenderGraphResourcePool.cpp b/Engine/Source/Runtime/RenderCore/Private/RenderGraphResourcePool.cpp
--- a/Engine/Source/Runtime/RenderCore/Private/RenderGraphResourcePool.cpp
+++ b/Engine/Source/Runtime/RenderCore/Private/RenderGraphResourcePool.cpp
@@ -12,11 +12,7 @@ uint64 ComputeHash(const FRDGBufferDesc& Desc)
 	return CityHash64((const char*)&Desc, sizeof(FRDGBufferDesc));
 }
 
-FRenderGraphResourcePool::FRenderGraphResourcePool()
-{ }
-
-
-TRefCountPtr<FRDGPooledBuffer> FRenderGraphResourcePool::FindFreeBuffer(
+TRefCountPtr<FRDGPooledBuffer> FRDGBufferPool::FindFreeBuffer(
 	FRHICommandList& RHICmdList,
 	const FRDGBufferDesc& Desc,
 	const TCHAR* InDebugName)
@@ -26,7 +22,7 @@ TRefCountPtr<FRDGPooledBuffer> FRenderGraphResourcePool::FindFreeBuffer(
 	return Result;
 }
 
-TRefCountPtr<FRDGPooledBuffer> FRenderGraphResourcePool::FindFreeBufferInternal(
+TRefCountPtr<FRDGPooledBuffer> FRDGBufferPool::FindFreeBufferInternal(
 	FRHICommandList& RHICmdList,
 	const FRDGBufferDesc& Desc,
 	const TCHAR* InDebugName)
@@ -36,7 +32,7 @@ TRefCountPtr<FRDGPooledBuffer> FRenderGraphResourcePool::FindFreeBufferInternal(
 	FRDGBufferDesc AlignedDesc = Desc;
 	AlignedDesc.NumElements = Align(AlignedDesc.BytesPerElement * AlignedDesc.NumElements, BufferPageSize) / AlignedDesc.BytesPerElement;
 	const uint64 BufferHash = ComputeHash(AlignedDesc);
-
+	
 	// First find if available.
 	for (int32 Index = 0; Index < AllocatedBufferHashes.Num(); ++Index)
 	{
@@ -71,7 +67,7 @@ TRefCountPtr<FRDGPooledBuffer> FRenderGraphResourcePool::FindFreeBufferInternal(
 
 	// Allocate new one
 	{
-		TRACE_CPUPROFILER_EVENT_SCOPE(FRenderGraphResourcePool::CreateBuffer);
+		TRACE_CPUPROFILER_EVENT_SCOPE(FRDGBufferPool::CreateBuffer);
 
 		const uint32 NumBytes = AlignedDesc.GetTotalNumBytes();
 
@@ -110,12 +106,12 @@ TRefCountPtr<FRDGPooledBuffer> FRenderGraphResourcePool::FindFreeBufferInternal(
 	}
 }
 
-void FRenderGraphResourcePool::ReleaseDynamicRHI()
+void FRDGBufferPool::ReleaseDynamicRHI()
 {
 	AllocatedBuffers.Empty();
 }
 
-void FRenderGraphResourcePool::TickPoolElements()
+void FRDGBufferPool::TickPoolElements()
 {
 	const uint32 kFramesUntilRelease = 30;
 
@@ -143,4 +139,148 @@ void FRenderGraphResourcePool::TickPoolElements()
 	++FrameCounter;
 }
 
-TGlobalResource<FRenderGraphResourcePool> GRenderGraphResourcePool;
+TGlobalResource<FRDGBufferPool> GRenderGraphResourcePool;
+
+uint32 FRDGTransientRenderTarget::AddRef() const
+{
+	check(IsInRenderingThread());
+	check(LifetimeState == ERDGTransientResourceLifetimeState::Allocated);
+	return ++RefCount;
+}
+
+uint32 FRDGTransientRenderTarget::Release()
+{
+	check(IsInRenderingThread());
+	check(RefCount > 0 && LifetimeState == ERDGTransientResourceLifetimeState::Allocated);
+	const uint32 Refs = --RefCount;
+	if (Refs == 0)
+	{
+		GRDGTransientResourceAllocator.AddPendingDeallocation(this);
+	}
+	return Refs;
+}
+
+void FRDGTransientResourceAllocator::InitDynamicRHI()
+{
+	Allocator = RHICreateTransientResourceAllocator();
+}
+
+void FRDGTransientResourceAllocator::ReleaseDynamicRHI()
+{
+	if (Allocator)
+	{
+		FRHICommandListImmediate& RHICmdList = FRHICommandListExecutor::GetImmediateCommandList();
+
+		ReleasePendingDeallocations();
+		PendingDeallocationList.Empty();
+
+		for (FRDGTransientRenderTarget* RenderTarget : DeallocatedList)
+		{
+			delete RenderTarget;
+		}
+		DeallocatedList.Empty();
+
+		Allocator->Flush(RHICmdList);
+		Allocator->Release(RHICmdList);
+		Allocator = nullptr;
+	}
+}
+
+TRefCountPtr<FRDGTransientRenderTarget> FRDGTransientResourceAllocator::AllocateRenderTarget(FRHITransientTexture* Texture)
+{
+	check(Texture);
+
+	FRDGTransientRenderTarget* RenderTarget = nullptr;
+
+	if (!FreeList.IsEmpty())
+	{
+		RenderTarget = FreeList.Pop();
+	}
+	else
+	{
+		RenderTarget = new FRDGTransientRenderTarget();
+	}
+
+	RenderTarget->Texture = Texture;
+	RenderTarget->Desc = Translate(Texture->CreateInfo);
+	RenderTarget->Desc.DebugName = Texture->GetName();
+	RenderTarget->LifetimeState = ERDGTransientResourceLifetimeState::Allocated;
+	RenderTarget->GetRenderTargetItem().TargetableTexture = Texture->GetRHI();
+	RenderTarget->GetRenderTargetItem().ShaderResourceTexture = Texture->GetRHI();
+	InitAsWholeResource(RenderTarget->State, {});
+	return RenderTarget;
+}
+
+void FRDGTransientResourceAllocator::Release(TRefCountPtr<FRDGTransientRenderTarget>&& RenderTarget, FRDGPassHandle PassHandle)
+{
+	check(RenderTarget);
+
+	if (RenderTarget->GetRefCount() == 1)
+	{
+		Allocator->DeallocateMemory(RenderTarget->Texture, PassHandle.GetIndex());
+		RenderTarget->Reset();
+		RenderTarget = nullptr;
+	}
+}
+
+void FRDGTransientResourceAllocator::AddPendingDeallocation(FRDGTransientRenderTarget* RenderTarget)
+{
+	check(RenderTarget);
+	check(RenderTarget->GetRefCount() == 0);
+
+	if (RenderTarget->Texture)
+	{
+		RenderTarget->LifetimeState = ERDGTransientResourceLifetimeState::PendingDeallocation;
+		PendingDeallocationList.Emplace(RenderTarget);
+	}
+	else
+	{
+		RenderTarget->LifetimeState = ERDGTransientResourceLifetimeState::Deallocated;
+		DeallocatedList.Emplace(RenderTarget);
+	}
+}
+
+void FRDGTransientResourceAllocator::ReleasePendingDeallocations()
+{
+	if (!PendingDeallocationList.IsEmpty())
+	{
+		FMemStack& MemStack = FMemStack::Get();
+		FMemMark Mark(MemStack);
+
+		TArray<FRHITransitionInfo, TMemStackAllocator<>> Transitions;
+		Transitions.Reserve(PendingDeallocationList.Num());
+
+		TArray<FRHITransientAliasingInfo, TMemStackAllocator<>> Aliases;
+		Aliases.Reserve(PendingDeallocationList.Num());
+
+		for (FRDGTransientRenderTarget* RenderTarget : PendingDeallocationList)
+		{
+			Allocator->DeallocateMemory(RenderTarget->Texture, 0);
+
+			Aliases.Emplace(FRHITransientAliasingInfo::Discard(RenderTarget->Texture->GetRHI()));
+			Transitions.Emplace(RenderTarget->Texture->GetRHI(), ERHIAccess::Unknown, ERHIAccess::Discard);
+
+			RenderTarget->Reset();
+			RenderTarget->LifetimeState = ERDGTransientResourceLifetimeState::Deallocated;
+		}
+
+		{
+			const FRHITransition* Transition = RHICreateTransition(FRHITransitionCreateInfo(ERHIPipeline::Graphics, ERHIPipeline::Graphics, ERHITransitionCreateFlags::None, Transitions, Aliases));
+
+			FRHICommandListImmediate& RHICmdList = FRHICommandListExecutor::GetImmediateCommandList();
+			RHICmdList.BeginTransition(Transition);
+			RHICmdList.EndTransition(Transition);
+		}
+
+		FreeList.Append(PendingDeallocationList);
+		PendingDeallocationList.Reset();
+	}
+
+	if (!DeallocatedList.IsEmpty())
+	{
+		FreeList.Append(DeallocatedList);
+		DeallocatedList.Reset();
+	}
+}
+
+TGlobalResource<FRDGTransientResourceAllocator> GRDGTransientResourceAllocator;
\ No newline at end of file
diff --git a/Engine/Source/Runtime/RenderCore/Private/RenderGraphResourcePool.h b/Engine/Source/Runtime/RenderCore/Private/RenderGraphResourcePool.h
--- a/Engine/Source/Runtime/RenderCore/Private/RenderGraphResourcePool.h
+++ b/Engine/Source/Runtime/RenderCore/Private/RenderGraphResourcePool.h
@@ -12,17 +12,10 @@
 #include "RendererInterface.h"
 #include "RenderGraphResources.h"
 
-
-/**
- * Pools all resources for the render graph.
- */
-class RENDERCORE_API FRenderGraphResourcePool : public FRenderResource
+class RENDERCORE_API FRDGBufferPool : public FRenderResource
 {
 public:
-	FRenderGraphResourcePool();
-
-	/** Free renderer resources */
-	virtual void ReleaseDynamicRHI() override;
+	FRDGBufferPool() = default;
 
 	/** Call once per frame to trim elements from the pool. */
 	void TickPoolElements();
@@ -31,6 +24,8 @@ public:
 	TRefCountPtr<FRDGPooledBuffer> FindFreeBuffer(FRHICommandList& RHICmdList, const FRDGBufferDesc& Desc, const TCHAR* InDebugName);
 
 private:
+	void ReleaseDynamicRHI() override;
+
 	/** Allocate a buffer from a given descriptor. */
 	TRefCountPtr<FRDGPooledBuffer> FindFreeBufferInternal(FRHICommandList& RHICmdList, const FRDGBufferDesc& Desc, const TCHAR* InDebugName);
 
@@ -44,4 +39,80 @@ private:
 };
 
 /** The global render targets for easy shading. */
-extern RENDERCORE_API TGlobalResource<FRenderGraphResourcePool> GRenderGraphResourcePool;
+extern RENDERCORE_API TGlobalResource<FRDGBufferPool> GRenderGraphResourcePool;
+
+enum class ERDGTransientResourceLifetimeState
+{
+	Deallocated,
+	Allocated,
+	PendingDeallocation
+};
+
+class FRDGTransientRenderTarget final : public IPooledRenderTarget
+{
+public:
+	uint32 AddRef() const override;
+	uint32 Release() override;
+	uint32 GetRefCount() const override { return RefCount; }
+
+	bool IsFree() const override { return false; }
+	bool IsTracked() const override { return true; }
+	uint32 ComputeMemorySize() const override { return 0; }
+
+	void SetDebugName(const TCHAR*) override {}
+
+	const FPooledRenderTargetDesc& GetDesc() const override { return Desc; }
+
+	FRHITransientTexture* GetTransientTexture() const override
+	{
+		check(LifetimeState == ERDGTransientResourceLifetimeState::Allocated);
+		return Texture;
+	}
+
+	void Reset()
+	{
+		Texture = nullptr;
+		RenderTargetItem.ShaderResourceTexture = nullptr;
+		RenderTargetItem.TargetableTexture = nullptr;
+	}
+
+	FRDGTextureSubresourceState State;
+
+private:
+	FRDGTransientRenderTarget() = default;
+
+	FRHITransientTexture* Texture;
+	FPooledRenderTargetDesc Desc;
+	ERDGTransientResourceLifetimeState LifetimeState;
+	mutable uint32 RefCount = 0;
+
+	friend class FRDGTransientResourceAllocator;
+};
+
+class FRDGTransientResourceAllocator : public FRenderResource
+{
+public:
+	IRHITransientResourceAllocator* Get() { return Allocator; }
+
+	TRefCountPtr<FRDGTransientRenderTarget> AllocateRenderTarget(FRHITransientTexture* Texture);
+
+	void Release(TRefCountPtr<FRDGTransientRenderTarget>&& RenderTarget, FRDGPassHandle PassHandle);
+
+	void ReleasePendingDeallocations();
+
+private:
+	void InitDynamicRHI() override;
+	void ReleaseDynamicRHI() override;
+
+	void AddPendingDeallocation(FRDGTransientRenderTarget* RenderTarget);
+
+	IRHITransientResourceAllocator* Allocator = nullptr;
+
+	TArray<FRDGTransientRenderTarget*> FreeList;
+	TArray<FRDGTransientRenderTarget*> PendingDeallocationList;
+	TArray<FRDGTransientRenderTarget*> DeallocatedList;
+
+	friend class FRDGTransientRenderTarget;
+};
+
+extern RENDERCORE_API TGlobalResource<FRDGTransientResourceAllocator> GRDGTransientResourceAllocator;
\ No newline at end of file
diff --git a/Engine/Source/Runtime/RenderCore/Private/RenderGraphResources.cpp b/Engine/Source/Runtime/RenderCore/Private/RenderGraphResources.cpp
--- a/Engine/Source/Runtime/RenderCore/Private/RenderGraphResources.cpp
+++ b/Engine/Source/Runtime/RenderCore/Private/RenderGraphResources.cpp
@@ -3,6 +3,7 @@
 #include "RenderGraphResources.h"
 #include "RenderGraphPass.h"
 #include "RenderGraphPrivate.h"
+#include "RenderGraphResourcePool.h"
 
 inline bool NeedsUAVBarrier(FRDGViewHandle PreviousHandle, FRDGViewHandle NextHandle)
 {
@@ -27,7 +28,8 @@ FRDGParentResource::FRDGParentResource(const TCHAR* InName, const ERDGParentReso
 	, bExtracted(0)
 	, bProduced(0)
 	, bTransient(0)
-	, bUserSetNonTransient(0)
+	, bForceNonTransient(0)
+	, bAllowTransientExtracted(0)
 	, bFinalizedAccess(0)
 	, bLastOwner(1)
 	, bCulled(1)
@@ -206,17 +208,29 @@ IPooledRenderTarget* FRDGTexture::GetPooledRenderTarget() const
 	return PooledRenderTarget;
 }
 
-void FRDGTexture::SetRHI(FPooledRenderTarget* InPooledRenderTarget)
+void FRDGTexture::SetRHI(IPooledRenderTarget* InPooledRenderTarget)
 {
-	Allocation = TRefCountPtr<FPooledRenderTarget>(InPooledRenderTarget);
-	PooledRenderTarget = InPooledRenderTarget;
+	if (FRHITransientTexture* LocalTransientTexture = InPooledRenderTarget->GetTransientTexture())
+	{
+		FRDGTransientRenderTarget* LocalRenderTarget = static_cast<FRDGTransientRenderTarget*>(InPooledRenderTarget);
+		Allocation = TRefCountPtr<FRDGTransientRenderTarget>(LocalRenderTarget);
 
-	if (!InPooledRenderTarget->HasRDG())
+		SetRHI(LocalTransientTexture, &LocalRenderTarget->State);
+	}
+	else
 	{
-		InPooledRenderTarget->InitRDG();
+		FPooledRenderTarget* LocalRenderTarget = static_cast<FPooledRenderTarget*>(InPooledRenderTarget);
+		Allocation = TRefCountPtr<FPooledRenderTarget>(LocalRenderTarget);
+
+		if (!LocalRenderTarget->HasRDG())
+		{
+			LocalRenderTarget->InitRDG();
+		}
+
+		SetRHI(LocalRenderTarget->GetRDG(RenderTargetTexture));
 	}
 
-	SetRHI(InPooledRenderTarget->GetRDG(RenderTargetTexture));
+	PooledRenderTarget = InPooledRenderTarget;
 }
 
 void FRDGTexture::SetRHI(FRDGPooledTexture* InPooledTexture)
@@ -239,13 +253,12 @@ void FRDGTexture::SetRHI(FRDGPooledTexture* InPooledTexture)
 	ResourceRHI = PooledTexture->GetRHI();
 }
 
-void FRDGTexture::SetRHI(FRHITransientTexture* InTransientTexture, FRDGAllocator& Allocator)
+void FRDGTexture::SetRHI(FRHITransientTexture* InTransientTexture, FRDGTextureSubresourceState* InTransientTextureState)
 {
 	TransientTexture = InTransientTexture;
-	State = Allocator.AllocNoDestruct<FRDGTextureSubresourceState>();
+	State = InTransientTextureState;
 	ViewCache = &InTransientTexture->ViewCache;
 	ResourceRHI = InTransientTexture->GetRHI();
-
 	bTransient = true;
 }
 
@@ -258,20 +271,31 @@ void FRDGTexture::Finalize(FRDGPooledTextureArray& PooledTextureArray)
 	{
 		if (bTransient)
 		{
-			// Manually deconstruct the allocated state so as not to invoke overhead from the allocators destructor tracking.
-			State->~FRDGTextureSubresourceState();
-			State = nullptr;
-		}
-		else
-		{
-			// External and extracted resources are user controlled, so we cannot assume the texture stays in its final state.
-			if (bExternal || bExtracted)
+			if (PooledRenderTarget)
 			{
-				PooledTexture->Reset();
+				InitAsWholeResource(*State, {});
+				PooledTextureArray.Emplace(MoveTemp(Allocation));
 			}
 			else
 			{
-				PooledTexture->Finalize();
+				// Manually deconstruct the allocated state so as not to invoke overhead from the allocators destructor tracking.
+				State->~FRDGTextureSubresourceState();
+				State = nullptr;
+			}
+		}
+		else
+		{
+			if (PooledTexture)
+			{
+				// External and extracted resources are user controlled, so we cannot assume the texture stays in its final state.
+				if (bExternal || bExtracted)
+				{
+					PooledTexture->Reset();
+				}
+				else
+				{
+					PooledTexture->Finalize();
+				}
 			}
 
 			// Restore the reference to the last owner in the aliasing chain; not necessary for the transient resource allocator.
@@ -282,8 +306,7 @@ void FRDGTexture::Finalize(FRDGPooledTextureArray& PooledTextureArray)
 		}
 	}
 
-	// This releases the reference without invoking a virtual function call.
-	TRefCountPtr<FPooledRenderTarget>(MoveTemp(Allocation));
+	Allocation = nullptr;
 }
 
 void FRDGBuffer::SetRHI(FRDGPooledBuffer* InPooledBuffer)
diff --git a/Engine/Source/Runtime/RenderCore/Private/RenderGraphTrace.cpp b/Engine/Source/Runtime/RenderCore/Private/RenderGraphTrace.cpp
--- a/Engine/Source/Runtime/RenderCore/Private/RenderGraphTrace.cpp
+++ b/Engine/Source/Runtime/RenderCore/Private/RenderGraphTrace.cpp
@@ -14,8 +14,8 @@ UE_TRACE_EVENT_BEGIN(RDGTrace, GraphMessage)
 	UE_TRACE_EVENT_FIELD(uint64, StartCycles)
 	UE_TRACE_EVENT_FIELD(uint64, EndCycles)
 	UE_TRACE_EVENT_FIELD(uint16, PassCount)
-	UE_TRACE_EVENT_FIELD(uint64[], TransientHeapWatermarkSizes)
-	UE_TRACE_EVENT_FIELD(uint64[], TransientHeapCapacities)
+	UE_TRACE_EVENT_FIELD(uint64[], TransientMemoryCommitSizes)
+	UE_TRACE_EVENT_FIELD(uint64[], TransientMemoryCapacities)
 UE_TRACE_EVENT_END()
 
 UE_TRACE_EVENT_BEGIN(RDGTrace, GraphEndMessage)
@@ -52,13 +52,15 @@ UE_TRACE_EVENT_BEGIN(RDGTrace, BufferMessage)
 	UE_TRACE_EVENT_FIELD(uint16, NextOwnerHandle)
 	UE_TRACE_EVENT_FIELD(uint16, Order)
 	UE_TRACE_EVENT_FIELD(uint16[], Passes)
-	UE_TRACE_EVENT_FIELD(uint16, TransientHeapIndex)
-	UE_TRACE_EVENT_FIELD(uint64, TransientHeapOffsetMin)
-	UE_TRACE_EVENT_FIELD(uint64, TransientHeapOffsetMax)
+	UE_TRACE_EVENT_FIELD(uint64[], TransientAllocationOffsetMins)
+	UE_TRACE_EVENT_FIELD(uint64[], TransientAllocationOffsetMaxs)
+	UE_TRACE_EVENT_FIELD(uint16[], TransientAllocationMemoryRanges)
 	UE_TRACE_EVENT_FIELD(bool, IsExternal)
 	UE_TRACE_EVENT_FIELD(bool, IsExtracted)
 	UE_TRACE_EVENT_FIELD(bool, IsCulled)
 	UE_TRACE_EVENT_FIELD(bool, IsTransient)
+	UE_TRACE_EVENT_FIELD(bool, IsTransientUntracked)
+	UE_TRACE_EVENT_FIELD(bool, IsTransientCacheHit)
 UE_TRACE_EVENT_END()
  
 UE_TRACE_EVENT_BEGIN(RDGTrace, TextureMessage)
@@ -69,9 +71,9 @@ UE_TRACE_EVENT_BEGIN(RDGTrace, TextureMessage)
 	UE_TRACE_EVENT_FIELD(uint16, NextOwnerHandle)
 	UE_TRACE_EVENT_FIELD(uint16, Order)
 	UE_TRACE_EVENT_FIELD(uint16[], Passes)
-	UE_TRACE_EVENT_FIELD(uint16, TransientHeapIndex)
-	UE_TRACE_EVENT_FIELD(uint64, TransientHeapOffsetMin)
-	UE_TRACE_EVENT_FIELD(uint64, TransientHeapOffsetMax)
+	UE_TRACE_EVENT_FIELD(uint64[], TransientAllocationOffsetMins)
+	UE_TRACE_EVENT_FIELD(uint64[], TransientAllocationOffsetMaxs)
+	UE_TRACE_EVENT_FIELD(uint16[], TransientAllocationMemoryRanges)
 	UE_TRACE_EVENT_FIELD(uint64, SizeInBytes)
 	UE_TRACE_EVENT_FIELD(uint64, CreateFlags)
 	UE_TRACE_EVENT_FIELD(uint32, Dimension)
@@ -86,6 +88,8 @@ UE_TRACE_EVENT_BEGIN(RDGTrace, TextureMessage)
 	UE_TRACE_EVENT_FIELD(bool, IsExtracted)
 	UE_TRACE_EVENT_FIELD(bool, IsCulled)
 	UE_TRACE_EVENT_FIELD(bool, IsTransient)
+	UE_TRACE_EVENT_FIELD(bool, IsTransientUntracked)
+	UE_TRACE_EVENT_FIELD(bool, IsTransientCacheHit)
 UE_TRACE_EVENT_END()
 
 UE_TRACE_EVENT_BEGIN(RDGTrace, ScopeMessage)
@@ -99,14 +103,18 @@ static_assert(sizeof(FRDGPassHandle) == sizeof(uint16), "Expected 16 bit pass ha
 static_assert(sizeof(FRDGTextureHandle) == sizeof(uint16), "Expected 16 bit texture handles.");
 static_assert(sizeof(FRDGBufferHandle) == sizeof(uint16), "Expected 16 bit buffer handles.");
 
-bool IsTraceEnabled()
+FRDGTrace::FRDGTrace()
+	: bEnabled(UE_TRACE_CHANNELEXPR_IS_ENABLED(RDGChannel) && !IsImmediateMode())
+{}
+
+bool FRDGTrace::IsEnabled() const
 {
-	return UE_TRACE_CHANNELEXPR_IS_ENABLED(RDGChannel) && !IsImmediateMode();
+	return bEnabled;
 }
 
 void FRDGTrace::OutputGraphBegin()
 {
-	if (!IsTraceEnabled())
+	if (!IsEnabled())
 	{
 		return;
 	}
@@ -116,7 +124,7 @@ void FRDGTrace::OutputGraphBegin()
 
 void FRDGTrace::OutputGraphEnd(const FRDGBuilder& GraphBuilder)
 {
-	if (!IsTraceEnabled())
+	if (!IsEnabled())
 	{
 		return;
 	}
@@ -132,13 +140,16 @@ void FRDGTrace::OutputGraphEnd(const FRDGBuilder& GraphBuilder)
 	{
 		const TCHAR* Name = GraphBuilder.BuilderName.GetTCHAR();
 
-		TArray<uint64, TInlineAllocator<8>> TransientHeapWatermarkSizes;
-		TArray<uint64, TInlineAllocator<8>> TransientHeapCapacities;
+		TArray<uint64> TransientMemoryCommitSizes;
+		TArray<uint64> TransientMemoryCapacities;
+
+		TransientMemoryCommitSizes.Reserve(TransientAllocationStats.MemoryRanges.Num());
+		TransientMemoryCapacities.Reserve(TransientAllocationStats.MemoryRanges.Num());
 
-		for (const auto& Heap : TransientHeapStats.Heaps)
+		for (const auto& MemoryRange : TransientAllocationStats.MemoryRanges)
 		{
-			TransientHeapWatermarkSizes.Emplace(Heap.WatermarkSize);
-			TransientHeapCapacities.Emplace(Heap.Capacity);
+			TransientMemoryCommitSizes.Emplace(MemoryRange.CommitSize);
+			TransientMemoryCapacities.Emplace(MemoryRange.Capacity);
 		}
 
 		UE_TRACE_LOG(RDGTrace, GraphMessage, RDGChannel)
@@ -146,8 +157,8 @@ void FRDGTrace::OutputGraphEnd(const FRDGBuilder& GraphBuilder)
 			<< GraphMessage.StartCycles(GraphStartCycles)
 			<< GraphMessage.EndCycles(FPlatformTime::Cycles64())
 			<< GraphMessage.PassCount(uint16(Passes.Num()))
-			<< GraphMessage.TransientHeapWatermarkSizes(TransientHeapWatermarkSizes.GetData(), (uint16)TransientHeapWatermarkSizes.Num())
-			<< GraphMessage.TransientHeapCapacities(TransientHeapCapacities.GetData(), (uint16)TransientHeapCapacities.Num());
+			<< GraphMessage.TransientMemoryCommitSizes(TransientMemoryCommitSizes.GetData(), (uint16)TransientMemoryCommitSizes.Num())
+			<< GraphMessage.TransientMemoryCapacities(TransientMemoryCapacities.GetData(), (uint16)TransientMemoryCapacities.Num());
 	}
 
 	for (FRDGPassHandle Handle = Passes.Begin(); Handle != Passes.End(); ++Handle)
@@ -248,6 +259,37 @@ void FRDGTrace::OutputGraphEnd(const FRDGBuilder& GraphBuilder)
 	}
 #endif
 
+	struct FTransientAllocation
+	{
+		TArray<uint64> OffsetMins;
+		TArray<uint64> OffsetMaxs;
+		TArray<uint16> MemoryRanges;
+		bool bCacheHit = false;
+
+		void Reset()
+		{
+			OffsetMins.Reset();
+			OffsetMaxs.Reset();
+			MemoryRanges.Reset();
+			bCacheHit = false;
+		}
+
+		void Fill(const FRHITransientAllocationStats& Stats, const FRHITransientResource* Resource)
+		{
+			const FRHITransientAllocationStats::FAllocationArray& Allocations = Stats.Resources.FindChecked(Resource);
+
+			for (const FRHITransientAllocationStats::FAllocation& Allocation : Allocations)
+			{
+				OffsetMins.Emplace(Allocation.OffsetMin);
+				OffsetMaxs.Emplace(Allocation.OffsetMax);
+				MemoryRanges.Emplace(Allocation.MemoryRangeIndex);
+			}
+
+			bCacheHit = Resource->GetAcquireCount() > 1;
+		}
+
+	} TransientAllocation;
+
 	for (FRDGTextureHandle Handle = Textures.Begin(); Handle != Textures.End(); ++Handle)
 	{
 		const FRDGTexture* Texture = Textures[Handle];
@@ -258,11 +300,13 @@ void FRDGTrace::OutputGraphEnd(const FRDGBuilder& GraphBuilder)
 			SizeInBytes = RHIComputeMemorySize(TextureRHI);
 		}
 
-		FRHITransientResourceStats TransientStats;
+		TransientAllocation.Reset();
 
 		if (Texture->bTransient)
 		{
-			TransientStats = Texture->TransientTexture->GetStats();
+			const FRHITransientResource* Resource = Texture->TransientTexture;
+			TransientAllocation.Fill(TransientAllocationStats, Resource);
+			TransientAllocationStats.Resources.Remove(Resource);
 		}
 
 		UE_TRACE_LOG(RDGTrace, TextureMessage, RDGChannel)
@@ -271,9 +315,9 @@ void FRDGTrace::OutputGraphEnd(const FRDGBuilder& GraphBuilder)
 			<< TextureMessage.NextOwnerHandle(Texture->NextOwner.GetIndexUnchecked())
 			<< TextureMessage.Order(Texture->TraceOrder)
 			<< TextureMessage.Passes((const uint16*)Texture->TracePasses.GetData(), (uint16)Texture->TracePasses.Num())
-			<< TextureMessage.TransientHeapIndex(TransientStats.HeapIndex)
-			<< TextureMessage.TransientHeapOffsetMin(TransientStats.HeapOffsetMin)
-			<< TextureMessage.TransientHeapOffsetMax(TransientStats.HeapOffsetMax)
+			<< TextureMessage.TransientAllocationOffsetMins(TransientAllocation.OffsetMins.GetData(), TransientAllocation.OffsetMins.Num())
+			<< TextureMessage.TransientAllocationOffsetMaxs(TransientAllocation.OffsetMaxs.GetData(), TransientAllocation.OffsetMaxs.Num())
+			<< TextureMessage.TransientAllocationMemoryRanges(TransientAllocation.MemoryRanges.GetData(), TransientAllocation.MemoryRanges.Num())
 			<< TextureMessage.SizeInBytes(SizeInBytes)
 			<< TextureMessage.CreateFlags(uint32(Texture->Desc.Flags))
 			<< TextureMessage.Dimension(uint32(Texture->Desc.Dimension))
@@ -287,18 +331,22 @@ void FRDGTrace::OutputGraphEnd(const FRDGBuilder& GraphBuilder)
 			<< TextureMessage.IsExternal(bool(Texture->bExternal))
 			<< TextureMessage.IsExtracted(bool(Texture->bExtracted))
 			<< TextureMessage.IsCulled(bool(Texture->bCulled))
-			<< TextureMessage.IsTransient(bool(Texture->bTransient));
+			<< TextureMessage.IsTransient(bool(Texture->bTransient))
+			<< TextureMessage.IsTransientUntracked(false)
+			<< TextureMessage.IsTransientCacheHit(TransientAllocation.bCacheHit);
 	}
 
 	for (FRDGBufferHandle Handle = Buffers.Begin(); Handle != Buffers.End(); ++Handle)
 	{
 		const FRDGBuffer* Buffer = Buffers[Handle];
 
-		FRHITransientResourceStats TransientStats;
+		TransientAllocation.Reset();
 
 		if (Buffer->bTransient)
 		{
-			TransientStats = Buffer->TransientBuffer->GetStats();
+			const FRHITransientResource* Resource = Buffer->TransientBuffer;
+			TransientAllocation.Fill(TransientAllocationStats, Resource);
+			TransientAllocationStats.Resources.Remove(Resource);
 		}
 
 		UE_TRACE_LOG(RDGTrace, BufferMessage, RDGChannel)
@@ -307,16 +355,86 @@ void FRDGTrace::OutputGraphEnd(const FRDGBuilder& GraphBuilder)
 			<< BufferMessage.NextOwnerHandle(Buffer->NextOwner.GetIndexUnchecked())
 			<< BufferMessage.Order(Buffer->TraceOrder)
 			<< BufferMessage.Passes((const uint16*)Buffer->TracePasses.GetData(), (uint16)Buffer->TracePasses.Num())
-			<< BufferMessage.TransientHeapIndex(TransientStats.HeapIndex)
-			<< BufferMessage.TransientHeapOffsetMin(TransientStats.HeapOffsetMin)
-			<< BufferMessage.TransientHeapOffsetMax(TransientStats.HeapOffsetMax)
+			<< BufferMessage.TransientAllocationOffsetMins(TransientAllocation.OffsetMins.GetData(), TransientAllocation.OffsetMins.Num())
+			<< BufferMessage.TransientAllocationOffsetMaxs(TransientAllocation.OffsetMaxs.GetData(), TransientAllocation.OffsetMaxs.Num())
+			<< BufferMessage.TransientAllocationMemoryRanges(TransientAllocation.MemoryRanges.GetData(), TransientAllocation.MemoryRanges.Num())
 			<< BufferMessage.UsageFlags(uint32(Buffer->Desc.Usage))
 			<< BufferMessage.BytesPerElement(Buffer->Desc.BytesPerElement)
 			<< BufferMessage.NumElements(Buffer->Desc.NumElements)
 			<< BufferMessage.IsExternal(bool(Buffer->bExternal))
 			<< BufferMessage.IsExtracted(bool(Buffer->bExtracted))
 			<< BufferMessage.IsCulled(bool(Buffer->bCulled))
-			<< BufferMessage.IsTransient(bool(Buffer->bTransient));
+			<< BufferMessage.IsTransient(bool(Buffer->bTransient))
+			<< BufferMessage.IsTransientUntracked(false)
+			<< BufferMessage.IsTransientCacheHit(TransientAllocation.bCacheHit);
+	}
+
+	int32 TextureIndex = Textures.Num();
+	int32 BufferIndex = Buffers.Num();
+
+	for (auto KeyValue : TransientAllocationStats.Resources)
+	{
+		const FRHITransientResource* Resource = KeyValue.Key;
+
+		if (!Resource->IsAcquired())
+		{
+			continue;
+		}
+
+		TransientAllocation.Reset();
+		TransientAllocation.Fill(TransientAllocationStats, Resource);
+
+		if (Resource->GetResourceType() == ERHITransientResourceType::Texture)
+		{
+			const FRHITransientTexture* Texture = static_cast<const FRHITransientTexture*>(Resource);
+
+			UE_TRACE_LOG(RDGTrace, TextureMessage, RDGChannel)
+				<< TextureMessage.Name(Texture->GetName(), uint16(FCString::Strlen(Texture->GetName())))
+				<< TextureMessage.Handle(TextureIndex)
+				<< TextureMessage.TransientAllocationOffsetMins(TransientAllocation.OffsetMins.GetData(), TransientAllocation.OffsetMins.Num())
+				<< TextureMessage.TransientAllocationOffsetMaxs(TransientAllocation.OffsetMaxs.GetData(), TransientAllocation.OffsetMaxs.Num())
+				<< TextureMessage.TransientAllocationMemoryRanges(TransientAllocation.MemoryRanges.GetData(), TransientAllocation.MemoryRanges.Num())
+				<< TextureMessage.SizeInBytes(Resource->GetSize())
+				<< TextureMessage.CreateFlags(uint32(Texture->CreateInfo.Flags))
+				<< TextureMessage.Dimension(uint32(Texture->CreateInfo.Dimension))
+				<< TextureMessage.Format(uint32(Texture->CreateInfo.Format))
+				<< TextureMessage.ExtentX(Texture->CreateInfo.Extent.X)
+				<< TextureMessage.ExtentY(Texture->CreateInfo.Extent.Y)
+				<< TextureMessage.Depth(Texture->CreateInfo.Depth)
+				<< TextureMessage.ArraySize(Texture->CreateInfo.ArraySize)
+				<< TextureMessage.NumMips(Texture->CreateInfo.NumMips)
+				<< TextureMessage.NumSamples(Texture->CreateInfo.NumSamples)
+				<< TextureMessage.IsExternal(false)
+				<< TextureMessage.IsExtracted(false)
+				<< TextureMessage.IsCulled(false)
+				<< TextureMessage.IsTransient(true)
+				<< TextureMessage.IsTransientUntracked(true)
+				<< TextureMessage.IsTransientCacheHit(TransientAllocation.bCacheHit);
+
+			TextureIndex++;
+		}
+		else
+		{
+			const FRHITransientBuffer* Buffer = static_cast<const FRHITransientBuffer*>(Resource);
+
+			UE_TRACE_LOG(RDGTrace, BufferMessage, RDGChannel)
+				<< BufferMessage.Name(Buffer->GetName(), uint16(FCString::Strlen(Buffer->GetName())))
+				<< BufferMessage.Handle(BufferIndex)
+				<< BufferMessage.TransientAllocationOffsetMins(TransientAllocation.OffsetMins.GetData(), TransientAllocation.OffsetMins.Num())
+				<< BufferMessage.TransientAllocationOffsetMaxs(TransientAllocation.OffsetMaxs.GetData(), TransientAllocation.OffsetMaxs.Num())
+				<< BufferMessage.TransientAllocationMemoryRanges(TransientAllocation.MemoryRanges.GetData(), TransientAllocation.MemoryRanges.Num())
+				<< BufferMessage.UsageFlags(uint32(Buffer->CreateInfo.Usage))
+				<< BufferMessage.BytesPerElement(Buffer->CreateInfo.Stride)
+				<< BufferMessage.NumElements(Buffer->CreateInfo.Size / Buffer->CreateInfo.Stride)
+				<< BufferMessage.IsExternal(false)
+				<< BufferMessage.IsExtracted(false)
+				<< BufferMessage.IsCulled(false)
+				<< BufferMessage.IsTransient(true)
+				<< BufferMessage.IsTransientUntracked(true)
+				<< BufferMessage.IsTransientCacheHit(TransientAllocation.bCacheHit);
+
+			BufferIndex++;
+		}
 	}
 
 	UE_TRACE_LOG(RDGTrace, GraphEndMessage, RDGChannel);
@@ -329,7 +447,7 @@ void FRDGTrace::AddResource(FRDGParentResource* Resource)
 
 void FRDGTrace::AddTexturePassDependency(FRDGTexture* Texture, FRDGPass* Pass)
 {
-	if (!IsTraceEnabled())
+	if (!IsEnabled())
 	{
 		return;
 	}
@@ -340,7 +458,7 @@ void FRDGTrace::AddTexturePassDependency(FRDGTexture* Texture, FRDGPass* Pass)
 
 void FRDGTrace::AddBufferPassDependency(FRDGBuffer* Buffer, FRDGPass* Pass)
 {
-	if (!IsTraceEnabled())
+	if (!IsEnabled())
 	{
 		return;
 	}
diff --git a/Engine/Source/Runtime/RenderCore/Private/RenderGraphValidation.cpp b/Engine/Source/Runtime/RenderCore/Private/RenderGraphValidation.cpp
--- a/Engine/Source/Runtime/RenderCore/Private/RenderGraphValidation.cpp
+++ b/Engine/Source/Runtime/RenderCore/Private/RenderGraphValidation.cpp
@@ -240,7 +240,6 @@ void FRDGUserValidation::ValidateRegisterExternalTexture(
 {
 	checkf(Name, TEXT("Attempted to register external texture with NULL name."));
 	checkf(ExternalPooledTexture.IsValid(), TEXT("Attempted to register NULL external texture."));
-	checkf(ExternalPooledTexture->IsCompatibleWithRDG(), TEXT("Pooled render target %s is not a compatible type for RDG."), Name);
 	ExecuteGuard(TEXT("RegisterExternalTexture"), Name);
 }
 
@@ -421,10 +420,6 @@ void FRDGUserValidation::ValidateExtractResource(FRDGParentResourceRef Resource)
 {
 	check(Resource);
 
-	checkf(!bHasExecuteBegun || !Resource->bTransient,
-		TEXT("Unable to queue the extraction of the resource %s because passes in the graph have already executed and it was made transient."),
-		Resource->Name);
-
 	checkf(Resource->bProduced || Resource->bExternal || Resource->bQueuedForUpload,
 		TEXT("Unable to queue the extraction of the resource %s because it has not been produced by any pass."),
 		Resource->Name);
diff --git a/Engine/Source/Runtime/RenderCore/Public/RenderGraphBuilder.h b/Engine/Source/Runtime/RenderCore/Public/RenderGraphBuilder.h
--- a/Engine/Source/Runtime/RenderCore/Public/RenderGraphBuilder.h
+++ b/Engine/Source/Runtime/RenderCore/Public/RenderGraphBuilder.h
@@ -214,8 +214,8 @@ public:
 	 *  the lifetime of the GPU resource until execution, at which point the pointer is filled. If specified, the texture is transitioned
 	 *  to the AccessFinal state, or kDefaultAccessFinal otherwise.
 	 */
-	void QueueTextureExtraction(FRDGTextureRef Texture, TRefCountPtr<IPooledRenderTarget>* OutPooledTexturePtr);
-	void QueueTextureExtraction(FRDGTextureRef Texture, TRefCountPtr<IPooledRenderTarget>* OutPooledTexturePtr, ERHIAccess AccessFinal);
+	void QueueTextureExtraction(FRDGTextureRef Texture, TRefCountPtr<IPooledRenderTarget>* OutPooledTexturePtr, ERDGResourceExtractionFlags Flags = ERDGResourceExtractionFlags::None);
+	void QueueTextureExtraction(FRDGTextureRef Texture, TRefCountPtr<IPooledRenderTarget>* OutPooledTexturePtr, ERHIAccess AccessFinal, ERDGResourceExtractionFlags Flags = ERDGResourceExtractionFlags::None);
 
 	/** Queues a pooled buffer extraction to happen at the end of graph execution. For graph-created buffers, this extends the lifetime
 	 *  of the GPU resource until execution, at which point the pointer is filled. If specified, the buffer is transitioned to the
@@ -626,26 +626,7 @@ private:
 	IF_RDG_CMDLIST_STATS(TStatId CommandListStatScope);
 	IF_RDG_CMDLIST_STATS(TStatId CommandListStatState);
 
-	class FTransientResourceAllocator
-	{
-	public:
-		FTransientResourceAllocator(FRHICommandListImmediate& InRHICmdList)
-			: RHICmdList(InRHICmdList)
-		{}
-
-		~FTransientResourceAllocator();
-
-		IRHITransientResourceAllocator* GetOrCreate();
-		IRHITransientResourceAllocator* Get() const { return Allocator; }
-		IRHITransientResourceAllocator* operator->() const { check(Allocator);  return Allocator; }
-		operator bool() const { return Allocator != nullptr; }
-
-	private:
-		FRHICommandListImmediate& RHICmdList;
-		IRHITransientResourceAllocator* Allocator = nullptr;
-		bool bCreateAttempted = false;
-
-	} TransientResourceAllocator;
+	IRHITransientResourceAllocator* TransientResourceAllocator = nullptr;
 
 	enum class EExecuteMode
 	{
diff --git a/Engine/Source/Runtime/RenderCore/Public/RenderGraphBuilder.inl b/Engine/Source/Runtime/RenderCore/Public/RenderGraphBuilder.inl
--- a/Engine/Source/Runtime/RenderCore/Public/RenderGraphBuilder.inl
+++ b/Engine/Source/Runtime/RenderCore/Public/RenderGraphBuilder.inl
@@ -268,30 +268,31 @@ inline void FRDGBuilder::QueueBufferUpload(FRDGBufferRef Buffer, FRDGBufferIniti
 	Buffer->bQueuedForUpload = 1;
 }
 
-inline void FRDGBuilder::QueueTextureExtraction(FRDGTextureRef Texture, TRefCountPtr<IPooledRenderTarget>* OutTexturePtr, ERHIAccess AccessFinal)
+inline void FRDGBuilder::QueueTextureExtraction(FRDGTextureRef Texture, TRefCountPtr<IPooledRenderTarget>* OutTexturePtr, ERHIAccess AccessFinal, ERDGResourceExtractionFlags Flags)
 {
-	QueueTextureExtraction(Texture, OutTexturePtr);
-
-	if (Texture->bTransient)
-	{
-		return;
-	}
-
+	QueueTextureExtraction(Texture, OutTexturePtr, Flags);
 	SetTextureAccessFinal(Texture, AccessFinal);
 }
 
-inline void FRDGBuilder::QueueTextureExtraction(FRDGTextureRef Texture, TRefCountPtr<IPooledRenderTarget>* OutTexturePtr)
+inline void FRDGBuilder::QueueTextureExtraction(FRDGTextureRef Texture, TRefCountPtr<IPooledRenderTarget>* OutTexturePtr, ERDGResourceExtractionFlags Flags)
 {
 	IF_RDG_ENABLE_DEBUG(UserValidation.ValidateExtractTexture(Texture, OutTexturePtr));
 
-	if (Texture->bTransient)
-	{
-		return;
-	}
+	*OutTexturePtr = nullptr;
 
 	Texture->ReferenceCount++;
 	Texture->bExtracted = true;
 	Texture->bCulled = false;
+
+	if (EnumHasAnyFlags(Flags, ERDGResourceExtractionFlags::AllowTransient))
+	{
+		Texture->bAllowTransientExtracted = 1;
+	}
+	else
+	{
+		Texture->bForceNonTransient = 1;
+	}
+
 	ExtractedTextures.Emplace(Texture, OutTexturePtr);
 
 	if (Texture->AccessFinal == ERHIAccess::Unknown)
@@ -304,9 +305,12 @@ inline void FRDGBuilder::QueueBufferExtraction(FRDGBufferRef Buffer, TRefCountPt
 {
 	IF_RDG_ENABLE_DEBUG(UserValidation.ValidateExtractBuffer(Buffer, OutBufferPtr));
 
+	*OutBufferPtr = nullptr;
+
 	Buffer->ReferenceCount++;
 	Buffer->bExtracted = true;
 	Buffer->bCulled = false;
+	Buffer->bForceNonTransient = true;
 	ExtractedBuffers.Emplace(Buffer, OutBufferPtr);
 
 	if (Buffer->AccessFinal == ERHIAccess::Unknown)
diff --git a/Engine/Source/Runtime/RenderCore/Public/RenderGraphDefinitions.h b/Engine/Source/Runtime/RenderCore/Public/RenderGraphDefinitions.h
--- a/Engine/Source/Runtime/RenderCore/Public/RenderGraphDefinitions.h
+++ b/Engine/Source/Runtime/RenderCore/Public/RenderGraphDefinitions.h
@@ -177,6 +177,17 @@ enum class ERDGViewType : uint8
 	MAX
 };
 
+enum class ERDGResourceExtractionFlags : uint8
+{
+	None = 0,
+
+	// Allows the resource to remain transient. Only use this flag if you intend to register the resource back
+	// into the graph and release the reference. This should not be used if the resource is cached for a long
+	// period of time.
+	AllowTransient = 1,
+};
+ENUM_CLASS_FLAGS(ERDGResourceExtractionFlags);
+
 enum class ERDGInitialDataFlags : uint8
 {
 	/** Specifies the default behavior, which is to make a copy of the initial data for replay when
@@ -522,7 +533,6 @@ class FRDGBarrierValidation;
 class FRDGBuilder;
 class FRDGEventName;
 class FRDGUserValidation;
-class FRenderGraphResourcePool;
 
 class FRDGResource;
 using FRDGResourceRef = FRDGResource*;
@@ -581,6 +591,8 @@ using FRDGBufferBitArray = TRDGHandleBitArray<FRDGBufferHandle>;
 
 class FRDGPooledTexture;
 class FRDGPooledBuffer;
+class FRDGBufferPool;
+class FRDGTransientRenderTarget;
 
 template <typename TUniformStruct> class TRDGUniformBuffer;
 template <typename TUniformStruct> using TRDGUniformBufferRef = TRDGUniformBuffer<TUniformStruct>*;
diff --git a/Engine/Source/Runtime/RenderCore/Public/RenderGraphResources.h b/Engine/Source/Runtime/RenderCore/Public/RenderGraphResources.h
--- a/Engine/Source/Runtime/RenderCore/Public/RenderGraphResources.h
+++ b/Engine/Source/Runtime/RenderCore/Public/RenderGraphResources.h
@@ -84,7 +84,7 @@ using FRDGTextureSubresourceState = TRDGTextureSubresourceArray<FRDGSubresourceS
 using FRDGTextureTransientSubresourceState = TRDGTextureSubresourceArray<FRDGSubresourceState, FRDGArrayAllocator>;
 using FRDGTextureTransientSubresourceStateIndirect = TRDGTextureSubresourceArray<FRDGSubresourceState*, FRDGArrayAllocator>;
 
-using FRDGPooledTextureArray = TArray<TRefCountPtr<FPooledRenderTarget>, FRDGArrayAllocator>;
+using FRDGPooledTextureArray = TArray<TRefCountPtr<IPooledRenderTarget>, FRDGArrayAllocator>;
 using FRDGPooledBufferArray = TArray<TRefCountPtr<FRDGPooledBuffer>, FRDGArrayAllocator>;
 
 /** Generic graph resource. */
@@ -254,6 +254,11 @@ public:
 		return bExtracted;
 	}
 
+	bool IsCulled() const
+	{
+		return bCulled;
+	}
+
 	/** Whether a prior pass added to the graph produced contents for this resource. External resources are not considered produced
 	 *  until used for a write operation. This is a union of all subresources, so any subresource write will set this to true.
 	 */
@@ -266,7 +271,7 @@ public:
 	void SetNonTransient()
 	{
 		check(!bTransient);
-		bUserSetNonTransient = 1;
+		bForceNonTransient = 1;
 	}
 
 protected:
@@ -284,8 +289,11 @@ protected:
 	/** Whether this resource is allocated through the transient resource allocator. */
 	uint8 bTransient : 1;
 
-	/** Whether this resource is set to be non-transient by the user. */
-	uint8 bUserSetNonTransient : 1;
+	/** Whether this resource cannot be made transient. */
+	uint8 bForceNonTransient : 1;
+
+	/** Whether this resource is allowed to be both transient and extracted. */
+	uint8 bAllowTransientExtracted : 1;
 
 	/** (External | Extracted only) If true, the resource is locked in its current state and will not be transitioned any more. */
 	uint8 bFinalizedAccess : 1;
@@ -482,14 +490,14 @@ private:
 		bSwapChain = EnumHasAnyFlags(Desc.Flags, ETextureCreateFlags::Presentable);
 	}
 
-	/** Assigns a pooled render target as the backing RHI resource. */
-	void SetRHI(FPooledRenderTarget* PooledRenderTarget);
+	/** Assigns a render target as the backing RHI resource. */
+	void SetRHI(IPooledRenderTarget* PooledRenderTarget);
 
-	/** Assigns a pooled buffer as the backing RHI resource. */
+	/** Assigns a pooled texture as the backing RHI resource. */
 	void SetRHI(FRDGPooledTexture* PooledTexture);
 
-	/** Assigns a transient buffer as the backing RHI resource. */
-	void SetRHI(FRHITransientTexture* TransientTexture, FRDGAllocator& Allocator);
+	/** Assigns a transient texture as the backing RHI resource. */
+	void SetRHI(FRHITransientTexture* TransientTexture, FRDGTextureSubresourceState* TransientTextureState);
 
 	/** Finalizes the texture for execution; no other transitions are allowed after calling this. */
 	void Finalize(FRDGPooledTextureArray& PooledTextureArray);
@@ -522,7 +530,7 @@ private:
 	const uint32 SubresourceCount;
 
 	/** The assigned pooled render target to use during execution. Never reset. */
-	FPooledRenderTarget* PooledRenderTarget = nullptr;
+	IPooledRenderTarget* PooledRenderTarget = nullptr;
 
 	union
 	{
@@ -994,7 +1002,7 @@ private:
 	const uint32 NumAllocatedElements;
 	uint32 LastUsedFrame = 0;
 
-	friend FRenderGraphResourcePool;
+	friend FRDGBufferPool;
 	friend FRDGBuilder;
 	friend FRDGBuffer;
 	friend FRDGAllocator;
diff --git a/Engine/Source/Runtime/RenderCore/Public/RenderGraphTrace.h b/Engine/Source/Runtime/RenderCore/Public/RenderGraphTrace.h
--- a/Engine/Source/Runtime/RenderCore/Public/RenderGraphTrace.h
+++ b/Engine/Source/Runtime/RenderCore/Public/RenderGraphTrace.h
@@ -12,6 +12,8 @@ UE_TRACE_CHANNEL_EXTERN(RDGChannel, RENDERCORE_API);
 class RENDERCORE_API FRDGTrace
 {
 public:
+	FRDGTrace();
+
 	void OutputGraphBegin();
 	void OutputGraphEnd(const FRDGBuilder& GraphBuilder);
 
@@ -19,12 +21,14 @@ public:
 	void AddTexturePassDependency(FRDGTexture* Texture, FRDGPass* Pass);
 	void AddBufferPassDependency(FRDGBuffer* Buffer, FRDGPass* Pass);
 
-	void SetTransientHeapStats(const FRHITransientHeapStats& Stats) { TransientHeapStats = Stats; }
+	FRHITransientAllocationStats TransientAllocationStats;
+
+	bool IsEnabled() const;
 
 private:
-	FRHITransientHeapStats TransientHeapStats;
 	uint64 GraphStartCycles{};
 	uint32 ResourceOrder{};
+	bool bEnabled;
 };
 
 #endif
\ No newline at end of file
diff --git a/Engine/Source/Runtime/RenderCore/Public/RenderTargetPool.h b/Engine/Source/Runtime/RenderCore/Public/RenderTargetPool.h
--- a/Engine/Source/Runtime/RenderCore/Public/RenderTargetPool.h
+++ b/Engine/Source/Runtime/RenderCore/Public/RenderTargetPool.h
@@ -68,10 +68,6 @@ struct RENDERCORE_API FPooledRenderTarget final : public IPooledRenderTarget
 	{
 		return RenderTargetPool != nullptr;
 	}
-	bool IsCompatibleWithRDG() const override
-	{
-		return true;
-	}
 
 	virtual void SetDebugName(const TCHAR *InName);
 	virtual const FPooledRenderTargetDesc& GetDesc() const;
diff --git a/Engine/Source/Runtime/RenderCore/Public/RendererInterface.h b/Engine/Source/Runtime/RenderCore/Public/RendererInterface.h
--- a/Engine/Source/Runtime/RenderCore/Public/RendererInterface.h
+++ b/Engine/Source/Runtime/RenderCore/Public/RendererInterface.h
@@ -41,6 +41,8 @@ namespace Nanite
 	struct FResources;
 };
 
+class FRHITransientTexture;
+
 // Shortcut for the allocator used by scene rendering.
 typedef TMemStackAllocator<> SceneRenderingAllocator;
 
@@ -553,10 +555,12 @@ struct IPooledRenderTarget
 	inline FSceneRenderTargetItem& GetRenderTargetItem() { return RenderTargetItem; }
 	/** Get the low level internals (texture/surface) */
 	inline const FSceneRenderTargetItem& GetRenderTargetItem() const { return RenderTargetItem; }
+
 	/** Returns if the render target is tracked by a pool. */
 	virtual bool IsTracked() const = 0;
-	/** Returns true if the render target is compatible with RDG. */
-	virtual bool IsCompatibleWithRDG() const { return false; }
+
+	/** Returns a transient texture if this is a container for one. */
+	virtual FRHITransientTexture* GetTransientTexture() const { return nullptr; }
 
 	// Refcounting
 	virtual uint32 AddRef() const = 0;
