UE5_RELEASE:  MGPU, fix for crash uploading staging buffers, when the staging buffers end up being pool allocated.  Need to properly handle aliasing logic for pool allocated versions of the staging resource.  Pool allocation is used when the buffers are especially large, which can randomly happen depending on what's going on in the frame, which is why it usually wouldn't hit this bug on a typical frame.

#rnx

diff --git a/Engine/Source/Runtime/D3D12RHI/Private/D3D12Buffer.cpp b/Engine/Source/Runtime/D3D12RHI/Private/D3D12Buffer.cpp
--- a/Engine/Source/Runtime/D3D12RHI/Private/D3D12Buffer.cpp
+++ b/Engine/Source/Runtime/D3D12RHI/Private/D3D12Buffer.cpp
@@ -722,17 +722,22 @@ void FD3D12DynamicRHI::UnlockBuffer(FRHICommandListImmediate* RHICmdList, FD3D12
 		{
 			// Update all of the resources in the LDA chain
 			check(Buffer->IsHeadLink());
+			FD3D12Buffer* LastBuffer = Buffer->GetLinkedObject(Buffer->GetLinkedObjectsGPUMask().GetLastIndex());
+
 			for (FD3D12Buffer::FLinkedObjectIterator CurrentBuffer(Buffer); CurrentBuffer; ++CurrentBuffer)
 			{
 				// If we are on the render thread, queue up the copy on the RHIThread so it happens at the correct time.
 				if (ShouldDeferBufferLockOperation(RHICmdList))
 				{
-					if (GNumExplicitGPUsForRendering == 1)
+					if (CurrentBuffer.Get() == LastBuffer)
 					{
+						// Command associated with last buffer (will be only buffer if single GPU) receives ownership of locked data
 						ALLOC_COMMAND_CL(*RHICmdList, FRHICommandUpdateBuffer)(&CurrentBuffer->ResourceLocation, LockedData.ResourceLocation, LockedData.LockedOffset, LockedData.LockedPitch);
 					}
-					else // The resource location must be copied because the constructor in FRHICommandUpdateBuffer transfers ownership and clears it.
+					else
 					{
+						// Other commands receive a reference copy of the locked data.  Commands get replayed in order, with the
+						// last command handling clean up the locked data after it has been propagated to all GPUs.
 						FD3D12ResourceLocation NodeResourceLocation(LockedData.ResourceLocation.GetParentDevice());
 						FD3D12ResourceLocation::ReferenceNode(NodeResourceLocation.GetParentDevice(), NodeResourceLocation, LockedData.ResourceLocation);
 						ALLOC_COMMAND_CL(*RHICmdList, FRHICommandUpdateBuffer)(&CurrentBuffer->ResourceLocation, NodeResourceLocation, LockedData.LockedOffset, LockedData.LockedPitch);
diff --git a/Engine/Source/Runtime/RHICore/Private/RHIPoolAllocator.cpp b/Engine/Source/Runtime/RHICore/Private/RHIPoolAllocator.cpp
--- a/Engine/Source/Runtime/RHICore/Private/RHIPoolAllocator.cpp
+++ b/Engine/Source/Runtime/RHICore/Private/RHIPoolAllocator.cpp
@@ -39,6 +39,7 @@ void FRHIPoolAllocationData::Reset()
 	Owner = nullptr;
 	PreviousAllocation = nullptr;
 	NextAllocation = nullptr;
+	AliasAllocation = nullptr;
 }
 
 void FRHIPoolAllocationData::InitAsHead(int16 InPoolIndex)
@@ -97,10 +98,32 @@ void FRHIPoolAllocationData::MoveFrom(FRHIPoolAllocationData& InAllocated, bool
 	Locked = InLocked;
 	
 	// Update linked list
-	InAllocated.PreviousAllocation->NextAllocation = this;
-	InAllocated.NextAllocation->PreviousAllocation = this;
-	PreviousAllocation = InAllocated.PreviousAllocation;
-	NextAllocation = InAllocated.NextAllocation;
+	if (InAllocated.PreviousAllocation)
+	{
+		InAllocated.PreviousAllocation->NextAllocation = this;
+		InAllocated.NextAllocation->PreviousAllocation = this;
+		PreviousAllocation = InAllocated.PreviousAllocation;
+		NextAllocation = InAllocated.NextAllocation;
+
+		// Update aliases (if present) to point to new parent
+		for (FRHIPoolAllocationData* Alias = InAllocated.GetFirstAlias(); Alias; Alias = Alias->GetNext())
+		{
+			Alias->AliasAllocation = this;
+		}
+		AliasAllocation = InAllocated.AliasAllocation;
+	}
+	else
+	{
+		// If there is no linked list, it means we are being passed an alias to a resource referenced by multiple GPUs.
+		// For this case, we need to remove "InAllocated" from its original parent, and replace it with "this".  This is
+		// a rare edge case, used in FD3D12DynamicRHI::UnlockBuffer in D3D12Buffer.cpp, where we are uploading a
+		// single CPU staging buffer to multiple GPU mirror copies.
+		FRHIPoolAllocationData* AliasParent = InAllocated.AliasAllocation;		// Get original parent of "InAllocated"
+		check(AliasParent);
+
+		InAllocated.RemoveAlias();				// Remove "InAllocated" as alias from original parent
+		AliasParent->AddAlias(this);			// Add "this" as alias to original parent
+	}
 
 	InAllocated.Reset();
 }
