Replaced legacy XB1 transient page allocator with RDG page table mapped transient allocator. Added support for ESRAM page pool.

diff --git a/Engine/Plugins/RenderGraphInsights/Source/RenderGraphInsights/Private/RenderGraphProvider.cpp b/Engine/Plugins/RenderGraphInsights/Source/RenderGraphInsights/Private/RenderGraphProvider.cpp
--- a/Engine/Plugins/RenderGraphInsights/Source/RenderGraphInsights/Private/RenderGraphProvider.cpp
+++ b/Engine/Plugins/RenderGraphInsights/Source/RenderGraphInsights/Private/RenderGraphProvider.cpp
@@ -155,7 +155,8 @@ FGraphPacket::FGraphPacket(TraceServices::ILinearAllocator& Allocator, const UE:
 
 	const auto CommitSizes = Context.EventData.GetArrayView<uint64>("TransientMemoryCommitSizes");
 	const auto Capacities  = Context.EventData.GetArrayView<uint64>("TransientMemoryCapacities");
-	check(CommitSizes.Num() == Capacities.Num());
+	const auto Flags       = Context.EventData.GetArrayView<uint8>("TransientMemoryFlags");
+	check(CommitSizes.Num() == Capacities.Num() && Capacities.Num() == Flags.Num());
 
 	uint64 CurrentOffset = 0;
 	TransientMemoryRangeByteOffsets.Reserve(CommitSizes.Num());
@@ -165,6 +166,7 @@ FGraphPacket::FGraphPacket(TraceServices::ILinearAllocator& Allocator, const UE:
 		auto& MemoryRange = TransientAllocationStats.MemoryRanges.AddDefaulted_GetRef();
 		MemoryRange.CommitSize = CommitSizes[Index];
 		MemoryRange.Capacity = Capacities[Index];
+		MemoryRange.Flags = (FRHITransientAllocationStats::EMemoryRangeFlags)Flags[Index];
 
 		TransientMemoryRangeByteOffsets.Emplace(CurrentOffset);
 		CurrentOffset += MemoryRange.CommitSize;
diff --git a/Engine/Plugins/RenderGraphInsights/Source/RenderGraphInsights/Private/RenderGraphTrack.cpp b/Engine/Plugins/RenderGraphInsights/Source/RenderGraphInsights/Private/RenderGraphTrack.cpp
--- a/Engine/Plugins/RenderGraphInsights/Source/RenderGraphInsights/Private/RenderGraphTrack.cpp
+++ b/Engine/Plugins/RenderGraphInsights/Source/RenderGraphInsights/Private/RenderGraphTrack.cpp
@@ -911,7 +911,7 @@ void FRenderGraphTrack::Draw(const ITimingTrackDrawContext& Context) const
 
 			const uint64 MiB = 1024 * 1024;
 
-			const auto DrawMibText = [&](float Depth, uint64 Offset, uint32 MemoryRangeIndex, uint32 MemoryRangeSizeInMB)
+			const auto DrawMibText = [&](float Depth, uint64 Offset, uint32 MemoryRangeIndex, uint32 MemoryRangeSizeInMB, FRHITransientAllocationStats::EMemoryRangeFlags MemoryRangeFlags)
 			{
 				FString Text;
 				FLinearColor TextColor = MBTextColor;
@@ -923,6 +923,12 @@ void FRenderGraphTrack::Draw(const ITimingTrackDrawContext& Context) const
 				else
 				{
 					Text = FString::Printf(TEXT("Memory Range [%d] %uMiB"), MemoryRangeIndex, MemoryRangeSizeInMB);
+
+					if (EnumHasAnyFlags(MemoryRangeFlags, FRHITransientAllocationStats::EMemoryRangeFlags::FastVRAM))
+					{
+						Text += TEXT("(FastVRAM)");
+					}
+
 					TextColor.A = 1.0f;
 				}
 
@@ -968,7 +974,7 @@ void FRenderGraphTrack::Draw(const ITimingTrackDrawContext& Context) const
 
 					if (bDrawText && bMajorTick)
 					{
-						LineStartX = DrawMibText(Depth, Offset, MemoryRangeIndex, TransientMemoryRange.Capacity / MiB);
+						LineStartX = DrawMibText(Depth, Offset, MemoryRangeIndex, TransientMemoryRange.Capacity / MiB, TransientMemoryRange.Flags);
 						LineWidth -= LineStartX - VisibleGraph.SlateX;
 					}
 
diff --git a/Engine/Source/Runtime/Apple/AGXRHI/Private/AGXRHI.cpp b/Engine/Source/Runtime/Apple/AGXRHI/Private/AGXRHI.cpp
--- a/Engine/Source/Runtime/Apple/AGXRHI/Private/AGXRHI.cpp
+++ b/Engine/Source/Runtime/Apple/AGXRHI/Private/AGXRHI.cpp
@@ -559,7 +559,6 @@ FAGXDynamicRHI::FAGXDynamicRHI(ERHIFeatureLevel::Type RequestedFeatureLevel)
 	GSupportsRenderTargetFormat_PF_G8 = false;
 	GRHISupportsTextureStreaming = true;
 	GSupportsWideMRT = true;
-	// GSupportsTransientResourceAliasing = FAGXCommandQueue::SupportsFeature(EAGXFeaturesHeaps) && FAGXCommandQueue::SupportsFeature(EAGXFeaturesFences);
 	GSupportsSeparateRenderTargetBlendState = (GMaxRHIFeatureLevel >= ERHIFeatureLevel::SM5);
 
 	GRHISupportsPipelineFileCache = true;
diff --git a/Engine/Source/Runtime/Apple/MetalRHI/Private/MetalRHI.cpp b/Engine/Source/Runtime/Apple/MetalRHI/Private/MetalRHI.cpp
--- a/Engine/Source/Runtime/Apple/MetalRHI/Private/MetalRHI.cpp
+++ b/Engine/Source/Runtime/Apple/MetalRHI/Private/MetalRHI.cpp
@@ -541,7 +541,6 @@ FMetalDynamicRHI::FMetalDynamicRHI(ERHIFeatureLevel::Type RequestedFeatureLevel)
 	GSupportsRenderTargetFormat_PF_G8 = false;
 	GRHISupportsTextureStreaming = true;
 	GSupportsWideMRT = true;
-	// GSupportsTransientResourceAliasing = FMetalCommandQueue::SupportsFeature(EMetalFeaturesHeaps) && FMetalCommandQueue::SupportsFeature(EMetalFeaturesFences);
 	GSupportsSeparateRenderTargetBlendState = (GMaxRHIFeatureLevel >= ERHIFeatureLevel::SM5);
 
 	GRHISupportsPipelineFileCache = true;
diff --git a/Engine/Source/Runtime/D3D12RHI/Private/D3D12Buffer.cpp b/Engine/Source/Runtime/D3D12RHI/Private/D3D12Buffer.cpp
--- a/Engine/Source/Runtime/D3D12RHI/Private/D3D12Buffer.cpp
+++ b/Engine/Source/Runtime/D3D12RHI/Private/D3D12Buffer.cpp
@@ -468,7 +468,6 @@ void FD3D12Buffer::Swap(FD3D12Buffer& Other)
 	check(!LockedData.bLocked && !Other.LockedData.bLocked);
 	FRHIBuffer::Swap(Other);
 	FD3D12BaseShaderResource::Swap(Other);
-	FD3D12TransientResource::Swap(Other);
 	FD3D12LinkedAdapterObject<FD3D12Buffer>::Swap(Other);
 }
 
diff --git a/Engine/Source/Runtime/D3D12RHI/Public/D3D12Resources.h b/Engine/Source/Runtime/D3D12RHI/Public/D3D12Resources.h
--- a/Engine/Source/Runtime/D3D12RHI/Public/D3D12Resources.h
+++ b/Engine/Source/Runtime/D3D12RHI/Public/D3D12Resources.h
@@ -813,12 +813,6 @@ protected:
 
 
 #if PLATFORM_WINDOWS || PLATFORM_HOLOLENS
-class FD3D12TransientResource
-{
-	// Nothing special for fast ram
-public:
-	void Swap(FD3D12TransientResource&) {}
-};
 class FD3D12FastClearResource
 {
 public:
@@ -832,7 +826,7 @@ public:
 
 
 /** The base class of resources that may be bound as shader resources (texture or buffer). */
-class FD3D12BaseShaderResource : public FD3D12DeviceChild, public FD3D12TransientResource, public IRefCountedObject
+class FD3D12BaseShaderResource : public FD3D12DeviceChild, public IRefCountedObject
 {
 protected:
 	FCriticalSection RenameListenersCS;
diff --git a/Engine/Source/Runtime/RHI/Public/RHI.h b/Engine/Source/Runtime/RHI/Public/RHI.h
--- a/Engine/Source/Runtime/RHI/Public/RHI.h
+++ b/Engine/Source/Runtime/RHI/Public/RHI.h
@@ -430,6 +430,7 @@ extern RHI_API bool GSupportsEfficientAsyncCompute;
 extern RHI_API bool GSupportsParallelOcclusionQueries;
 
 /** true if the RHI supports aliasing of transient resources */
+UE_DEPRECATED(5.1, "GSupportsTransientResourceAliasing has been deprecated. Transient resource aliasing is handled through IRHITransientResourceAllocator.")
 extern RHI_API bool GSupportsTransientResourceAliasing;
 
 /** true if the RHI requires a valid RT bound during UAV scatter operation inside the pixel shader */
diff --git a/Engine/Source/Runtime/RHI/Public/RHICommandList.h b/Engine/Source/Runtime/RHI/Public/RHICommandList.h
--- a/Engine/Source/Runtime/RHI/Public/RHICommandList.h
+++ b/Engine/Source/Runtime/RHI/Public/RHICommandList.h
@@ -4628,50 +4628,18 @@ public:
 
 	FORCEINLINE void AcquireTransientResource_RenderThread(FRHITexture* Texture)
 	{
-		if (!Texture->IsCommitted() )
-		{
-			if (GSupportsTransientResourceAliasing)
-			{
-				GDynamicRHI->RHIAcquireTransientResource_RenderThread(Texture);
-			}
-			Texture->SetCommitted(true);
-		}
 	}
 
 	FORCEINLINE void DiscardTransientResource_RenderThread(FRHITexture* Texture)
 	{
-		if (Texture->IsCommitted())
-		{
-			if (GSupportsTransientResourceAliasing)
-			{
-				GDynamicRHI->RHIDiscardTransientResource_RenderThread(Texture);
-			}
-			Texture->SetCommitted(false);
-		}
 	}
 
 	FORCEINLINE void AcquireTransientResource_RenderThread(FRHIBuffer* Buffer)
 	{
-		if (!Buffer->IsCommitted())
-		{
-			if (GSupportsTransientResourceAliasing)
-			{
-				GDynamicRHI->RHIAcquireTransientResource_RenderThread(Buffer);
-			}
-			Buffer->SetCommitted(true);
-		}
 	}
 
 	FORCEINLINE void DiscardTransientResource_RenderThread(FRHIBuffer* Buffer)
 	{
-		if (Buffer->IsCommitted())
-		{
-			if (GSupportsTransientResourceAliasing)
-			{
-				GDynamicRHI->RHIDiscardTransientResource_RenderThread(Buffer);
-			}
-			Buffer->SetCommitted(false);
-		}
 	}
 
 	FORCEINLINE bool GetRenderQueryResult(FRHIRenderQuery* RenderQuery, uint64& OutResult, bool bWait, uint32 GPUIndex = INDEX_NONE)
diff --git a/Engine/Source/Runtime/RHI/Public/RHIResources.h b/Engine/Source/Runtime/RHI/Public/RHIResources.h
--- a/Engine/Source/Runtime/RHI/Public/RHIResources.h
+++ b/Engine/Source/Runtime/RHI/Public/RHIResources.h
@@ -129,19 +129,6 @@ public:
 
 	static bool Bypass();
 
-	// Transient resource tracking
-	// We do this at a high level so we can catch errors even when transient resources are not supported
-	void SetCommitted(bool bInCommitted) 
-	{ 
-		check(IsInRenderingThread()); 
-		bCommitted = bInCommitted;
-	}
-	bool IsCommitted() const 
-	{ 
-		check(IsInRenderingThread());
-		return bCommitted;
-	}
-
 	bool IsValid() const
 	{
 		return AtomicFlags.IsValid(std::memory_order_relaxed);
diff --git a/Engine/Source/Runtime/RHI/Public/RHITransientResourceAllocator.h b/Engine/Source/Runtime/RHI/Public/RHITransientResourceAllocator.h
--- a/Engine/Source/Runtime/RHI/Public/RHITransientResourceAllocator.h
+++ b/Engine/Source/Runtime/RHI/Public/RHITransientResourceAllocator.h
@@ -323,6 +323,14 @@ public:
 
 	using FAllocationArray = TArray<FAllocation, TInlineAllocator<2>>;
 
+	enum class EMemoryRangeFlags
+	{
+		None = 0,
+
+		// The memory range references platform specific fast RAM.
+		FastVRAM = 1 << 0
+	};
+
 	struct FMemoryRange
 	{
 		// Number of bytes available for use in the memory range.
@@ -330,12 +338,17 @@ public:
 
 		// Number of bytes allocated for use in the memory range.
 		uint64 CommitSize = 0;
+
+		// Flags specified for this memory range.
+		EMemoryRangeFlags Flags = EMemoryRangeFlags::None;
 	};
 
 	TArray<FMemoryRange> MemoryRanges;
 	TMap<const FRHITransientResource*, FAllocationArray> Resources;
 };
 
+ENUM_CLASS_FLAGS(FRHITransientAllocationStats::EMemoryRangeFlags);
+
 class IRHITransientResourceAllocator
 {
 public:
diff --git a/Engine/Source/Runtime/RHICore/Private/RHICoreTransientResourceAllocator.cpp b/Engine/Source/Runtime/RHICore/Private/RHICoreTransientResourceAllocator.cpp
--- a/Engine/Source/Runtime/RHICore/Private/RHICoreTransientResourceAllocator.cpp
+++ b/Engine/Source/Runtime/RHICore/Private/RHICoreTransientResourceAllocator.cpp
@@ -1293,7 +1293,6 @@ void FRHITransientPagePool::Allocate(FAllocationContext& Context)
 
 		bool bMapPages = true;
 
-#if 0
 		if (AllocationIndex < Context.AllocationsBefore.Num())
 		{
 			const FRHITransientPagePoolAllocation& AllocationBefore = Context.AllocationsBefore[AllocationIndex];
@@ -1304,7 +1303,6 @@ void FRHITransientPagePool::Allocate(FAllocationContext& Context)
 				bMapPages = false;
 			}
 		}
-#endif
 
 		if (bMapPages)
 		{
@@ -1322,13 +1320,16 @@ void FRHITransientPagePool::Allocate(FAllocationContext& Context)
 
 void FRHITransientPagePool::Flush(FRHICommandListImmediate& RHICmdList)
 {
-	PageMapRequestCountMax = FMath::Max<uint32>(PageMapRequests.Num(), PageMapRequestCountMax);
-	PageSpanCountMax       = FMath::Max<uint32>(PageSpans.Num(),       PageSpanCountMax);
+	if (!PageMapRequests.IsEmpty())
+	{
+		PageMapRequestCountMax = FMath::Max<uint32>(PageMapRequests.Num(), PageMapRequestCountMax);
+		PageSpanCountMax       = FMath::Max<uint32>(PageSpans.Num(),       PageSpanCountMax);
 
-	Flush(RHICmdList, MoveTemp(PageMapRequests), MoveTemp(PageSpans));
+		Flush(RHICmdList, MoveTemp(PageMapRequests), MoveTemp(PageSpans));
 
-	PageMapRequests.Reserve(PageMapRequestCountMax);
-	PageSpans.Reserve(PageSpanCountMax);
+		PageMapRequests.Reserve(PageMapRequestCountMax);
+		PageSpans.Reserve(PageSpanCountMax);
+	}
 
 	OverlapTracker.Reset();
 }
@@ -1348,6 +1349,14 @@ FRHITransientPagePoolCache::~FRHITransientPagePoolCache()
 {
 	checkf(FreeList.Num() == LiveList.Num(), TEXT("PagePools are still being held by transient allocators."));
 	FreeList.Empty();
+
+	delete FastPagePool;
+	FastPagePool = nullptr;
+
+	for (FRHITransientPagePool* PagePool : LiveList)
+	{
+		delete PagePool;
+	}
 	LiveList.Empty();
 }
 
@@ -1383,6 +1392,23 @@ FRHITransientPagePool* FRHITransientPagePoolCache::Acquire()
 	return PagePool;
 }
 
+FRHITransientPagePool* FRHITransientPagePoolCache::GetFastPagePool()
+{
+	if (!FastPagePool)
+	{
+		LLM_SCOPE_BYTAG(RHITransientResources);
+		FastPagePool = CreateFastPagePool();
+
+		if (FastPagePool)
+		{
+			TotalMemoryCapacity += FastPagePool->GetCapacity();
+			return FastPagePool;
+		}
+	}
+
+	return nullptr;
+}
+
 void FRHITransientPagePoolCache::Forfeit(TConstArrayView<FRHITransientPagePool*> InForfeitedPagePools)
 {
 	FScopeLock Lock(&CriticalSection);
@@ -1486,12 +1512,9 @@ void FRHITransientResourcePageAllocator::AllocateMemoryInternal(FRHITransientRes
 {
 	FRHITransientPagePool::FAllocationContext AllocationContext(*Resource, PageSize);
 
-	if (bFastPoolRequested)
+	if (bFastPoolRequested && FastPagePool)
 	{
-		if (FRHITransientPagePool* FastPagePool = PagePoolCache.GetFastPagePool())
-		{
-			FastPagePool->Allocate(AllocationContext);
-		}
+		FastPagePool->Allocate(AllocationContext);
 	}
 
 	if (!AllocationContext.IsComplete())
@@ -1545,19 +1568,28 @@ void FRHITransientResourcePageAllocator::Flush(FRHICommandListImmediate& RHICmdL
 {
 	if (OutAllocationStats)
 	{
-		TMap<FRHITransientPagePool*, int32> PagePoolToIndex;
-		PagePoolToIndex.Reserve(PagePools.Num());
+		TMap<FRHITransientPagePool*, int32> PagePoolToMemoryRangeIndex;
+		PagePoolToMemoryRangeIndex.Reserve(PagePools.Num() + !!FastPagePool);
 
-		for (int32 Index = 0; Index < PagePools.Num(); ++Index)
+		const auto AddMemoryRange = [&](FRHITransientPagePool* PagePool, FRHITransientAllocationStats::EMemoryRangeFlags Flags)
 		{
-			FRHITransientPagePool* PagePool = PagePools[Index];
+			PagePoolToMemoryRangeIndex.Emplace(PagePool, OutAllocationStats->MemoryRanges.Num());
 
 			FRHITransientAllocationStats::FMemoryRange MemoryRange;
 			MemoryRange.Capacity = PagePool->GetCapacity();
 			MemoryRange.CommitSize = PagePool->GetCapacity();
+			MemoryRange.Flags = Flags;
 			OutAllocationStats->MemoryRanges.Emplace(MemoryRange);
+		};
+
+		if (FastPagePool)
+		{
+			AddMemoryRange(FastPagePool, FRHITransientAllocationStats::EMemoryRangeFlags::FastVRAM);
+		}
 
-			PagePoolToIndex.Emplace(PagePool, Index);
+		for (FRHITransientPagePool* PagePool : PagePools)
+		{
+			AddMemoryRange(PagePool, FRHITransientAllocationStats::EMemoryRangeFlags::None);
 		}
 
 		for (const FRHITransientTexture* Texture : Textures.GetAllocated())
@@ -1569,7 +1601,7 @@ void FRHITransientResourcePageAllocator::Flush(FRHICommandListImmediate& RHICmdL
 				FRHITransientAllocationStats::FAllocation Allocation;
 				Allocation.OffsetMin = PageSize *  PageSpan.Offset;
 				Allocation.OffsetMax = PageSize * (PageSpan.Offset + PageSpan.Count);
-				Allocation.MemoryRangeIndex = PagePoolToIndex[PagePool];
+				Allocation.MemoryRangeIndex = PagePoolToMemoryRangeIndex[PagePool];
 				Allocations.Emplace(Allocation);
 			});
 		}
@@ -1583,7 +1615,7 @@ void FRHITransientResourcePageAllocator::Flush(FRHICommandListImmediate& RHICmdL
 				FRHITransientAllocationStats::FAllocation Allocation;
 				Allocation.OffsetMin = PageSize *  PageSpan.Offset;
 				Allocation.OffsetMax = PageSize * (PageSpan.Offset + PageSpan.Count);
-				Allocation.MemoryRangeIndex = PagePoolToIndex[PagePool];
+				Allocation.MemoryRangeIndex = PagePoolToMemoryRangeIndex[PagePool];
 				Allocations.Emplace(Allocation);
 			});
 		}
@@ -1618,6 +1650,11 @@ void FRHITransientResourcePageAllocator::Flush(FRHICommandListImmediate& RHICmdL
 		TRACE_COUNTER_SET(TransientBufferCacheSize, Buffers.GetSize());
 	}
 
+	if (FastPagePool)
+	{
+		FastPagePool->Flush(RHICmdList);
+	}
+
 	for (FRHITransientPagePool* PagePool : PagePools)
 	{
 		PagePool->Flush(RHICmdList);
diff --git a/Engine/Source/Runtime/RHICore/Public/RHICoreTransientResourceAllocator.h b/Engine/Source/Runtime/RHICore/Public/RHICoreTransientResourceAllocator.h
--- a/Engine/Source/Runtime/RHICore/Public/RHICoreTransientResourceAllocator.h
+++ b/Engine/Source/Runtime/RHICore/Public/RHICoreTransientResourceAllocator.h
@@ -1013,6 +1013,9 @@ public:
 	// Called by the transient allocator to acquire a page pool from the cache.
 	FRHITransientPagePool* Acquire();
 
+	// Called by the transient allocator to return the fast page pool if it exists.
+	FRHITransientPagePool* GetFastPagePool();
+
 	// Called by the transient allocator to forfeit all acquired heaps back to the cache.
 	void Forfeit(TConstArrayView<FRHITransientPagePool*> PagePools);
 
@@ -1025,7 +1028,7 @@ private:
 	//! Platform API
 
 	// Called to access the dedicated fast VRAM page pool, if it exists on the platform.
-	virtual FRHITransientPagePool* GetFastPagePool() = 0;
+	virtual FRHITransientPagePool* CreateFastPagePool() { return nullptr; }
 
 	// Called when a new heap is being created and added to the pool.
 	virtual FRHITransientPagePool* CreatePagePool(const FRHITransientPagePool::FInitializer& Initializer) = 0;
@@ -1035,6 +1038,7 @@ private:
 	FRHITransientMemoryStats Stats;
 
 	FCriticalSection CriticalSection;
+	FRHITransientPagePool* FastPagePool = nullptr;
 	TArray<FRHITransientPagePool*> LiveList;
 	TArray<FRHITransientPagePool*> FreeList;
 	uint64 GarbageCollectCycle = 0;
@@ -1053,7 +1057,9 @@ public:
 		, Textures(InPagePoolCache.Initializer.TextureCacheSize, ResourceCacheGarbageCollectionLatency)
 		, Buffers(InPagePoolCache.Initializer.BufferCacheSize, ResourceCacheGarbageCollectionLatency)
 		, PageSize(PagePoolCache.Initializer.PageSize)
-	{}
+	{
+		FastPagePool = PagePoolCache.GetFastPagePool();
+	}
 
 	~FRHITransientResourcePageAllocator();
 
@@ -1114,7 +1120,9 @@ private:
 	TRHITransientResourceCache<FRHITransientBuffer> Buffers;
 
 	TArray<FRHITransientPagePool*> PagePools;
+	FRHITransientPagePool* FastPagePool = nullptr;
 	uint64 CurrentCycle = 0;
+	uint32 FirstNormalPagePoolIndex = 0;
 	uint32 DeallocationCount = 0;
 	uint32 PageSize = 0;
 	uint32 PageMapCount = 0;
diff --git a/Engine/Source/Runtime/RenderCore/Private/RenderGraphBuilder.cpp b/Engine/Source/Runtime/RenderCore/Private/RenderGraphBuilder.cpp
--- a/Engine/Source/Runtime/RenderCore/Private/RenderGraphBuilder.cpp
+++ b/Engine/Source/Runtime/RenderCore/Private/RenderGraphBuilder.cpp
@@ -458,7 +458,7 @@ ERDGPassFlags FRDGBuilder::OverridePassFlags(const TCHAR* PassName, ERDGPassFlag
 
 bool FRDGBuilder::IsTransient(FRDGBufferRef Buffer) const
 {
-	if (!IsTransientInternal(Buffer))
+	if (!IsTransientInternal(Buffer, EnumHasAnyFlags(Buffer->Desc.Usage, BUF_FastVRAM)))
 	{
 		return false;
 	}
@@ -468,20 +468,15 @@ bool FRDGBuilder::IsTransient(FRDGBufferRef Buffer) const
 		return false;
 	}
 
-	if (GRDGTransientAllocator == 2 && !EnumHasAnyFlags(Buffer->Desc.Usage, BUF_FastVRAM))
-	{
-		return false;
-	}
-
 	return EnumHasAnyFlags(Buffer->Desc.Usage, BUF_UnorderedAccess);
 }
 
 bool FRDGBuilder::IsTransient(FRDGTextureRef Texture) const
 {
-	return IsTransientInternal(Texture) && (GRDGTransientAllocator != 2 || EnumHasAnyFlags(Texture->Desc.Flags, TexCreate_FastVRAM));
+	return IsTransientInternal(Texture, EnumHasAnyFlags(Texture->Desc.Flags, ETextureCreateFlags::FastVRAM));
 }
 
-bool FRDGBuilder::IsTransientInternal(FRDGParentResourceRef Resource) const
+bool FRDGBuilder::IsTransientInternal(FRDGParentResourceRef Resource, bool bFastVRAM) const
 {
 	// Immediate mode can't use the transient allocator because we don't know if the user will extract the resource.
 	if (!GRDGTransientAllocator || IsImmediateMode())
@@ -489,22 +484,31 @@ bool FRDGBuilder::IsTransientInternal(FRDGParentResourceRef Resource) const
 		return false;
 	}
 
-	if (Resource->bForceNonTransient)
+	// FastVRAM resources are always transient regardless of extraction or other hints, since they are performance critical.
+	if (!bFastVRAM)
 	{
-		return false;
-	}
-
-	if (Resource->bExtracted)
-	{
-		if (GRDGTransientExtractedResources == 0)
+		if (GRDGTransientAllocator == 2)
 		{
 			return false;
 		}
-
-		if (!Resource->bAllowTransientExtracted)
+	
+		if (Resource->bForceNonTransient)
 		{
 			return false;
 		}
+
+		if (Resource->bExtracted)
+		{
+			if (GRDGTransientExtractedResources == 0)
+			{
+				return false;
+			}
+
+			if (GRDGTransientExtractedResources == 1 && Resource->TransientExtractionHint == FRDGParentResource::ETransientExtractionHint::Disable)
+			{
+				return false;
+			}
+		}
 	}
 
 #if RDG_ENABLE_DEBUG
diff --git a/Engine/Source/Runtime/RenderCore/Private/RenderGraphPrivate.cpp b/Engine/Source/Runtime/RenderCore/Private/RenderGraphPrivate.cpp
--- a/Engine/Source/Runtime/RenderCore/Private/RenderGraphPrivate.cpp
+++ b/Engine/Source/Runtime/RenderCore/Private/RenderGraphPrivate.cpp
@@ -324,7 +324,8 @@ FAutoConsoleVariableRef CVarRDGTransientExtractedResource(
 	TEXT("r.RDG.TransientExtractedResources"), GRDGTransientExtractedResources,
 	TEXT("RDG will allocate extracted resources as transient, unless explicitly marked non-transient by the user.")
 	TEXT(" 0: disables external transient resources;")
-	TEXT(" 1: enables external transient resources (default);"),
+	TEXT(" 1: enables external transient resources (default);")
+	TEXT(" 2: force enables all external transient resources (not recommended);"),
 	ECVF_RenderThreadSafe);
 
 int32 GRDGParallelExecute = 1;
diff --git a/Engine/Source/Runtime/RenderCore/Private/RenderGraphResources.cpp b/Engine/Source/Runtime/RenderCore/Private/RenderGraphResources.cpp
--- a/Engine/Source/Runtime/RenderCore/Private/RenderGraphResources.cpp
+++ b/Engine/Source/Runtime/RenderCore/Private/RenderGraphResources.cpp
@@ -29,7 +29,7 @@ FRDGParentResource::FRDGParentResource(const TCHAR* InName, const ERDGParentReso
 	, bProduced(0)
 	, bTransient(0)
 	, bForceNonTransient(0)
-	, bAllowTransientExtracted(0)
+	, TransientExtractionHint(ETransientExtractionHint::None)
 	, bFinalizedAccess(0)
 	, bLastOwner(1)
 	, bCulled(1)
diff --git a/Engine/Source/Runtime/RenderCore/Private/RenderGraphTrace.cpp b/Engine/Source/Runtime/RenderCore/Private/RenderGraphTrace.cpp
--- a/Engine/Source/Runtime/RenderCore/Private/RenderGraphTrace.cpp
+++ b/Engine/Source/Runtime/RenderCore/Private/RenderGraphTrace.cpp
@@ -16,6 +16,7 @@ UE_TRACE_EVENT_BEGIN(RDGTrace, GraphMessage)
 	UE_TRACE_EVENT_FIELD(uint16, PassCount)
 	UE_TRACE_EVENT_FIELD(uint64[], TransientMemoryCommitSizes)
 	UE_TRACE_EVENT_FIELD(uint64[], TransientMemoryCapacities)
+	UE_TRACE_EVENT_FIELD(uint8[],  TransientMemoryFlags)
 UE_TRACE_EVENT_END()
 
 UE_TRACE_EVENT_BEGIN(RDGTrace, GraphEndMessage)
@@ -142,14 +143,17 @@ void FRDGTrace::OutputGraphEnd(const FRDGBuilder& GraphBuilder)
 
 		TArray<uint64> TransientMemoryCommitSizes;
 		TArray<uint64> TransientMemoryCapacities;
+		TArray<uint8>  TransientMemoryFlags;
 
 		TransientMemoryCommitSizes.Reserve(TransientAllocationStats.MemoryRanges.Num());
 		TransientMemoryCapacities.Reserve(TransientAllocationStats.MemoryRanges.Num());
+		TransientMemoryFlags.Reserve(TransientAllocationStats.MemoryRanges.Num());
 
 		for (const auto& MemoryRange : TransientAllocationStats.MemoryRanges)
 		{
 			TransientMemoryCommitSizes.Emplace(MemoryRange.CommitSize);
 			TransientMemoryCapacities.Emplace(MemoryRange.Capacity);
+			TransientMemoryFlags.Emplace((uint8)MemoryRange.Flags);
 		}
 
 		UE_TRACE_LOG(RDGTrace, GraphMessage, RDGChannel)
@@ -158,7 +162,8 @@ void FRDGTrace::OutputGraphEnd(const FRDGBuilder& GraphBuilder)
 			<< GraphMessage.EndCycles(FPlatformTime::Cycles64())
 			<< GraphMessage.PassCount(uint16(Passes.Num()))
 			<< GraphMessage.TransientMemoryCommitSizes(TransientMemoryCommitSizes.GetData(), (uint16)TransientMemoryCommitSizes.Num())
-			<< GraphMessage.TransientMemoryCapacities(TransientMemoryCapacities.GetData(), (uint16)TransientMemoryCapacities.Num());
+			<< GraphMessage.TransientMemoryCapacities(TransientMemoryCapacities.GetData(), (uint16)TransientMemoryCapacities.Num())
+			<< GraphMessage.TransientMemoryFlags(TransientMemoryFlags.GetData(), (uint16)TransientMemoryFlags.Num());
 	}
 
 	for (FRDGPassHandle Handle = Passes.Begin(); Handle != Passes.End(); ++Handle)
diff --git a/Engine/Source/Runtime/RenderCore/Private/RenderTargetPool.cpp b/Engine/Source/Runtime/RenderCore/Private/RenderTargetPool.cpp
--- a/Engine/Source/Runtime/RenderCore/Private/RenderTargetPool.cpp
+++ b/Engine/Source/Runtime/RenderCore/Private/RenderTargetPool.cpp
@@ -160,7 +160,7 @@ static TAutoConsoleVariable<int32> CVarAllowMultipleAliasingDiscardsPerFrame(
 static TAutoConsoleVariable<int32> CVarRtPoolTransientMode(
 	TEXT("r.RenderTargetPool.TransientAliasingMode"),
 	2,
-	TEXT("Enables transient resource aliasing for rendertargets. Used only if GSupportsTransientResourceAliasing is true.\n")
+	TEXT("Enables transient resource aliasing for rendertargets.\n")
 	TEXT("0 : Disabled\n")
 	TEXT("1 : enable transient resource aliasing for fastVRam rendertargets\n")
 	TEXT("2 : enable transient resource aliasing for fastVRam rendertargets and those with a Transient hint. Best for memory usage - has some GPU cost (~0.2ms)\n")
@@ -228,10 +228,6 @@ FRenderTargetPool::FRenderTargetPool()
 // Logic for determining whether to make a rendertarget transient
 bool FRenderTargetPool::DoesTargetNeedTransienceOverride(ETextureCreateFlags Flags, ERenderTargetTransience TransienceHint)
 {
-	if (!GSupportsTransientResourceAliasing)
-	{
-		return false;
-	}
 	int32 AliasingMode = CVarRtPoolTransientMode.GetValueOnRenderThread();
 
 	// We only override transience if aliasing is supported and enabled, the format is suitable, and the target is not already transient
@@ -1177,7 +1173,7 @@ void FRenderTargetPool::DumpMemoryUsage(FOutputDevice& OutputDevice)
 			}
 
 			OutputDevice.Logf(
-				TEXT("  %6.3fMB %4dx%4d%s%s %2dmip(s) %s (%s) %s %s Unused frames: %d"),
+				TEXT("  %6.3fMB %4dx%4d%s%s %2dmip(s) %s (%s) %s Unused frames: %d"),
 				ElementAllocationInKB / 1024.0f,
 				Element->Desc.Extent.X,
 				Element->Desc.Extent.Y,
@@ -1187,7 +1183,6 @@ void FRenderTargetPool::DumpMemoryUsage(FOutputDevice& OutputDevice)
 				Element->Desc.DebugName,
 				GPixelFormats[Element->Desc.Format].Name,
 				Element->IsTransient() ? TEXT("(transient)") : TEXT(""),
-				GSupportsTransientResourceAliasing ? *FString::Printf(TEXT("Frames since last discard: %d"), GFrameNumberRenderThread - Element->FrameNumberLastDiscard) : TEXT(""),
 				Element->UnusedForNFrames
 			);
 		}
@@ -1207,7 +1202,7 @@ void FRenderTargetPool::DumpMemoryUsage(FOutputDevice& OutputDevice)
 		if (Element)
 		{
 			OutputDevice.Logf(
-				TEXT("  %6.3fMB %4dx%4d%s%s %2dmip(s) %s (%s) %s %s"),
+				TEXT("  %6.3fMB %4dx%4d%s%s %2dmip(s) %s (%s) %s"),
 				ComputeSizeInKB(*Element) / 1024.0f,
 				Element->Desc.Extent.X,
 				Element->Desc.Extent.Y,
@@ -1216,8 +1211,7 @@ void FRenderTargetPool::DumpMemoryUsage(FOutputDevice& OutputDevice)
 				Element->Desc.NumMips,
 				Element->Desc.DebugName,
 				GPixelFormats[Element->Desc.Format].Name,
-				Element->IsTransient() ? TEXT("(transient)") : TEXT(""),
-				GSupportsTransientResourceAliasing ? *FString::Printf(TEXT("Frames since last discard: %d"), GFrameNumberRenderThread - Element->FrameNumberLastDiscard) : TEXT("")
+				Element->IsTransient() ? TEXT("(transient)") : TEXT("")
 			);
 			uint32 SizeInKB = ComputeSizeInKB(*Element);
 			DeferredTotal += SizeInKB;
diff --git a/Engine/Source/Runtime/RenderCore/Public/RenderGraphBuilder.h b/Engine/Source/Runtime/RenderCore/Public/RenderGraphBuilder.h
--- a/Engine/Source/Runtime/RenderCore/Public/RenderGraphBuilder.h
+++ b/Engine/Source/Runtime/RenderCore/Public/RenderGraphBuilder.h
@@ -706,7 +706,7 @@ private:
 
 	bool IsTransient(FRDGTextureRef Texture) const;
 	bool IsTransient(FRDGBufferRef Buffer) const;
-	bool IsTransientInternal(FRDGParentResourceRef Resource) const;
+	bool IsTransientInternal(FRDGParentResourceRef Resource, bool bFastVRAM) const;
 
 	FRHIRenderPassInfo GetRenderPassInfo(const FRDGPass* Pass) const;
 
diff --git a/Engine/Source/Runtime/RenderCore/Public/RenderGraphBuilder.inl b/Engine/Source/Runtime/RenderCore/Public/RenderGraphBuilder.inl
--- a/Engine/Source/Runtime/RenderCore/Public/RenderGraphBuilder.inl
+++ b/Engine/Source/Runtime/RenderCore/Public/RenderGraphBuilder.inl
@@ -286,11 +286,14 @@ inline void FRDGBuilder::QueueTextureExtraction(FRDGTextureRef Texture, TRefCoun
 
 	if (EnumHasAnyFlags(Flags, ERDGResourceExtractionFlags::AllowTransient))
 	{
-		Texture->bAllowTransientExtracted = 1;
+		if (Texture->TransientExtractionHint != FRDGTexture::ETransientExtractionHint::Disable)
+		{
+			Texture->TransientExtractionHint = FRDGTexture::ETransientExtractionHint::Enable;
+		}
 	}
 	else
 	{
-		Texture->bForceNonTransient = 1;
+		Texture->TransientExtractionHint = FRDGTexture::ETransientExtractionHint::Disable;
 	}
 
 	ExtractedTextures.Emplace(Texture, OutTexturePtr);
diff --git a/Engine/Source/Runtime/RenderCore/Public/RenderGraphResources.h b/Engine/Source/Runtime/RenderCore/Public/RenderGraphResources.h
--- a/Engine/Source/Runtime/RenderCore/Public/RenderGraphResources.h
+++ b/Engine/Source/Runtime/RenderCore/Public/RenderGraphResources.h
@@ -282,6 +282,13 @@ public:
 protected:
 	FRDGParentResource(const TCHAR* InName, ERDGParentResourceType InType);
 
+	enum class ETransientExtractionHint
+	{
+		None,
+		Disable,
+		Enable
+	};
+
 	/** Whether this is an externally registered resource. */
 	uint8 bExternal : 1;
 
@@ -298,7 +305,7 @@ protected:
 	uint8 bForceNonTransient : 1;
 
 	/** Whether this resource is allowed to be both transient and extracted. */
-	uint8 bAllowTransientExtracted : 1;
+	ETransientExtractionHint TransientExtractionHint : 2;
 
 	/** (External | Extracted only) If true, the resource is locked in its current state and will not be transitioned any more. */
 	uint8 bFinalizedAccess : 1;
diff --git a/Engine/Source/Runtime/RenderCore/Public/RenderTargetPool.h b/Engine/Source/Runtime/RenderCore/Public/RenderTargetPool.h
--- a/Engine/Source/Runtime/RenderCore/Public/RenderTargetPool.h
+++ b/Engine/Source/Runtime/RenderCore/Public/RenderTargetPool.h
@@ -354,9 +354,6 @@ private:
 	// sorted by size
 	SMemoryStats ComputeView();
 
-	// can be optimized to not be generated each time
-	void GenerateVRamAllocationUsage(TArray<FVRamAllocation>& Out);
-
 	friend struct FPooledRenderTarget;
 	friend class FVisualizeTexture;
 	friend class FVisualizeTexturePresent;
diff --git a/Engine/Source/Runtime/Renderer/Private/SceneRendering.cpp b/Engine/Source/Runtime/Renderer/Private/SceneRendering.cpp
--- a/Engine/Source/Runtime/Renderer/Private/SceneRendering.cpp
+++ b/Engine/Source/Runtime/Renderer/Private/SceneRendering.cpp
@@ -322,12 +322,6 @@ static TAutoConsoleVariable<int32> CVarWideCustomResolve(
 	ECVF_RenderThreadSafe | ECVF_Scalability
 	);
 
-TAutoConsoleVariable<int32> CVarTransientResourceAliasing_Buffers(
-	TEXT("r.TransientResourceAliasing.Buffers"),
-	1,
-	TEXT("Enables transient resource aliasing for specified buffers. Used only if GSupportsTransientResourceAliasing is true.\n"),
-	ECVF_ReadOnly);
-
 static TAutoConsoleVariable<int32> CVarBasePassForceOutputsVelocity(
 	TEXT("r.BasePassForceOutputsVelocity"), 0,
 	TEXT("Force the base pass to compute motion vector, regardless of FPrimitiveUniformShaderParameters.")
