Fixed flashing artifact with transient allocator due to incorrect barrier placement. Fixed FastVRAM path in RDG to actually check if FastVRAM is supported.

diff --git a/Engine/Source/Runtime/RHICore/Private/RHICoreTransientResourceAllocator.cpp b/Engine/Source/Runtime/RHICore/Private/RHICoreTransientResourceAllocator.cpp
--- a/Engine/Source/Runtime/RHICore/Private/RHICoreTransientResourceAllocator.cpp
+++ b/Engine/Source/Runtime/RHICore/Private/RHICoreTransientResourceAllocator.cpp
@@ -200,6 +200,31 @@ void FRHITransientResourceOverlapTracker::Track(FRHITransientResource* Transient
 
 //////////////////////////////////////////////////////////////////////////
 
+void FRHITransientResourceOverlapTracker::Reset()
+{
+	TArray<FResourceRange> AcquiredRanges;
+	TArray<FRHITransientResource*> AcquiredResources;
+	AcquiredRanges.Reserve(ResourceRanges.Num());
+	AcquiredResources.Reserve(Resources.Num());
+
+	for (FResourceRange Range : ResourceRanges)
+	{
+		FRHITransientResource* Resource = Resources[Range.ResourceIndex];
+
+		if (Resource->IsAcquired())
+		{
+			Range.ResourceIndex = AcquiredResources.Num();
+			AcquiredRanges.Emplace(Range);
+			AcquiredResources.Emplace(Resource);
+		}
+	}
+
+	Swap(ResourceRanges, AcquiredRanges);
+	Swap(Resources, AcquiredResources);
+}
+
+//////////////////////////////////////////////////////////////////////////
+
 template <typename ResourceContainerType>
 void LogActiveResources(const ResourceContainerType& Resources, uint64 CurrentCycle)
 {
diff --git a/Engine/Source/Runtime/RHICore/Public/RHICoreTransientResourceAllocator.h b/Engine/Source/Runtime/RHICore/Public/RHICoreTransientResourceAllocator.h
--- a/Engine/Source/Runtime/RHICore/Public/RHICoreTransientResourceAllocator.h
+++ b/Engine/Source/Runtime/RHICore/Public/RHICoreTransientResourceAllocator.h
@@ -142,12 +142,7 @@ class RHICORE_API FRHITransientResourceOverlapTracker
 {
 public:
 	void Track(FRHITransientResource* InResource, uint32 PageOffsetMin, uint32 PageOffsetMax);
-
-	void Reset()
-	{
-		ResourceRanges.Reset();
-		Resources.Reset();
-	}
+	void Reset();
 
 private:
 	struct FResourceRange
diff --git a/Engine/Source/Runtime/RenderCore/Private/RenderGraphBuilder.cpp b/Engine/Source/Runtime/RenderCore/Private/RenderGraphBuilder.cpp
--- a/Engine/Source/Runtime/RenderCore/Private/RenderGraphBuilder.cpp
+++ b/Engine/Source/Runtime/RenderCore/Private/RenderGraphBuilder.cpp
@@ -485,7 +485,7 @@ bool FRDGBuilder::IsTransientInternal(FRDGParentResourceRef Resource, bool bFast
 	}
 
 	// FastVRAM resources are always transient regardless of extraction or other hints, since they are performance critical.
-	if (!bFastVRAM)
+	if (!bFastVRAM || !FPlatformMemory::SupportsFastVRAMMemory())
 	{
 		if (GRDGTransientAllocator == 2)
 		{
