Relaxed validation on transient resource pool shutdown due to non-deterministic static initialization order.

diff --git a/Engine/Source/Runtime/RHICore/Private/RHICoreTransientResourceAllocator.cpp b/Engine/Source/Runtime/RHICore/Private/RHICoreTransientResourceAllocator.cpp
--- a/Engine/Source/Runtime/RHICore/Private/RHICoreTransientResourceAllocator.cpp
+++ b/Engine/Source/Runtime/RHICore/Private/RHICoreTransientResourceAllocator.cpp
@@ -41,15 +41,6 @@ static FAutoConsoleVariableRef CVarRHITransientAllocatorGarbageCollectLatency(
 	TEXT("Amount of update cycles before memory is reclaimed."),
 	ECVF_ReadOnly);
 
-#if RHICORE_TRANSIENT_ALLOCATOR_DEBUG
-static int32 GRHITransientAllocatorLogCycles = 0;
-static FAutoConsoleVariableRef CVarRHITransientAllocatorLogCycles(
-	TEXT("RHI.TransientAllocator.LogCycles"),
-	GRHITransientAllocatorLogCycles,
-	TEXT("Logs resources if they have been allocated for more than X allocator Flush() cycles; use to identify and log persistent allocations."),
-	ECVF_RenderThreadSafe);
-#endif
-
 TRACE_DECLARE_INT_COUNTER(TransientResourceCreateCount, TEXT("TransientAllocator/ResourceCreateCount"));
 
 TRACE_DECLARE_INT_COUNTER(TransientTextureCreateCount, TEXT("TransientAllocator/TextureCreateCount"));
@@ -223,71 +214,6 @@ void FRHITransientResourceOverlapTracker::Reset()
 	Swap(Resources, AcquiredResources);
 }
 
-//////////////////////////////////////////////////////////////////////////
-
-template <typename ResourceContainerType>
-void LogActiveResources(const ResourceContainerType& Resources, uint64 CurrentCycle)
-{
-#if RHICORE_TRANSIENT_ALLOCATOR_DEBUG
-	if (!Resources.IsEmpty())
-	{
-		FString ErrorMsg = FString::Printf(
-			TEXT("---------------------------------------------------------------------------------------------------------\n")
-			TEXT("-- %d Transient resources are still being held outside the allocator. Memory has been leaked."),
-			Resources.Num());
-
-		for (FRHITransientResource* Resource : Resources)
-		{
-			ErrorMsg += FString::Printf(TEXT("\tName: %s, Age: %llu cycles."), Resource->GetName(), CurrentCycle - Resource->GetAcquireCycle());
-		}
-
-		UE_LOG(LogRHICore, Fatal,
-			TEXT("%s\n")
-			TEXT("---------------------------------------------------------------------------------------------------------\n"),
-			*ErrorMsg);
-	}
-#endif
-}
-
-template <typename ResourceContainerType>
-void LogLongLivedResources(const ResourceContainerType& Resources, uint64 CurrentCycle)
-{
-#if RHICORE_TRANSIENT_ALLOCATOR_DEBUG
-	if (GRHITransientAllocatorLogCycles <= 0 || Resources.IsEmpty())
-	{
-		return;
-	}
-
-	TArray<FRHITransientResource*> LongLivedResources;
-
-	for (FRHITransientResource* Resource : Resources)
-	{
-		if (CurrentCycle - Resource->GetAcquireCycle() > GRHITransientAllocatorLogCycles)
-		{
-			LongLivedResources.Emplace(Resource);
-		}
-	}
-
-	if (!LongLivedResources.IsEmpty())
-	{
-		UE_LOG(LogRHICore, Warning,
-			TEXT("---------------------------------------------------------------------------------------------------------\n")
-			TEXT("-- %d transient %ss have been allocated for more than %d cycles. Please mark them non-transient to reduce fragmentation."),
-			LongLivedResources.Num(),
-			LongLivedResources[0]->GetResourceType() == ERHITransientResourceType::Texture ? TEXT("texture") : TEXT("buffer"),
-			GRHITransientAllocatorLogCycles);
-
-		for (FRHITransientResource* Resource : LongLivedResources)
-		{
-			UE_LOG(LogRHICore, Warning, TEXT("\tName: %s, Age: %llu cycles."), Resource->GetName(), CurrentCycle - Resource->GetAcquireCycle());
-		}
-
-		UE_LOG(LogRHICore, Warning,
-			TEXT("---------------------------------------------------------------------------------------------------------\n"));
-	}
-#endif
-};
-
 //////////////////////////////////////////////////////////////////////////
 // Transient Resource Heap Allocator
 //////////////////////////////////////////////////////////////////////////
@@ -300,22 +226,6 @@ FRHITransientHeapAllocator::FRHITransientHeapAllocator(uint64 InCapacity, uint32
 	InsertRange(HeadHandle, 0, Capacity);
 }
 
-FRHITransientHeapAllocator::~FRHITransientHeapAllocator()
-{
-#if RHICORE_TRANSIENT_ALLOCATOR_DEBUG
-	checkf(AllocationCount == 0, TEXT("%d allocations remain on heap."), AllocationCount);
-	check(HeadHandle != InvalidRangeHandle);
-
-	FRangeHandle FirstFreeHandle = GetFirstFreeRangeHandle();
-	check(FirstFreeHandle != InvalidRangeHandle);
-
-	const FRange& FirstFreeRange = Ranges[FirstFreeHandle];
-	check(FirstFreeRange.NextFreeHandle == InvalidRangeHandle);
-	check(FirstFreeRange.Offset == 0);
-	check(FirstFreeRange.Size == GetCapacity());
-#endif
-}
-
 FRHITransientHeapAllocation FRHITransientHeapAllocator::Allocate(uint64 Size, uint32 Alignment)
 {
 	check(Size > 0);
@@ -648,7 +558,6 @@ void FRHITransientHeap::Flush(uint64 AllocatorCycle, FRHITransientMemoryStats& O
 
 		{
 			Textures.Forfeit(GFrameCounterRenderThread);
-			LogLongLivedResources(Textures.GetAllocated(), AllocatorCycle);
 
 			for (FRHITransientTexture* Texture : Textures.GetAllocated())
 			{
@@ -661,7 +570,6 @@ void FRHITransientHeap::Flush(uint64 AllocatorCycle, FRHITransientMemoryStats& O
 
 		{
 			Buffers.Forfeit(GFrameCounterRenderThread);
-			LogLongLivedResources(Buffers.GetAllocated(), AllocatorCycle);
 
 			for (FRHITransientBuffer* Buffer : Buffers.GetAllocated())
 			{
@@ -692,9 +600,12 @@ FRHITransientHeapCache::FInitializer FRHITransientHeapCache::FInitializer::Creat
 
 FRHITransientHeapCache::~FRHITransientHeapCache()
 {
-	checkf(FreeList.Num() == LiveList.Num(), TEXT("Heaps are still being held by transient allocators."));
-	FreeList.Empty();
+	for (FRHITransientHeap* Heap : LiveList)
+	{
+		delete Heap;
+	}
 	LiveList.Empty();
+	FreeList.Empty();
 }
 
 FRHITransientHeap* FRHITransientHeapCache::Acquire(uint64 FirstAllocationSize, ERHITransientHeapFlags FirstAllocationHeapFlags)
@@ -781,11 +692,6 @@ void FRHITransientHeapCache::GarbageCollect()
 
 //////////////////////////////////////////////////////////////////////////
 
-FRHITransientResourceHeapAllocator::~FRHITransientResourceHeapAllocator()
-{
-	IF_RHICORE_TRANSIENT_ALLOCATOR_DEBUG(LogActiveResources(ActiveResources, CurrentCycle));
-}
-
 FRHITransientTexture* FRHITransientResourceHeapAllocator::CreateTextureInternal(
 	const FRHITextureCreateInfo& CreateInfo,
 	const TCHAR* DebugName,
@@ -1372,9 +1278,6 @@ FRHITransientPagePoolCache::FInitializer FRHITransientPagePoolCache::FInitialize
 
 FRHITransientPagePoolCache::~FRHITransientPagePoolCache()
 {
-	checkf(FreeList.Num() == LiveList.Num(), TEXT("PagePools are still being held by transient allocators."));
-	FreeList.Empty();
-
 	delete FastPagePool;
 	FastPagePool = nullptr;
 
@@ -1383,6 +1286,7 @@ FRHITransientPagePoolCache::~FRHITransientPagePoolCache()
 		delete PagePool;
 	}
 	LiveList.Empty();
+	FreeList.Empty();
 }
 
 FRHITransientPagePool* FRHITransientPagePoolCache::Acquire()
@@ -1490,11 +1394,6 @@ void FRHITransientPagePoolCache::GarbageCollect()
 
 ///////////////////////////////////////////////////////////////////////////////////////////////////
 
-FRHITransientResourcePageAllocator::~FRHITransientResourcePageAllocator()
-{
-	IF_RHICORE_TRANSIENT_ALLOCATOR_DEBUG(LogActiveResources(ActiveResources, CurrentCycle));
-}
-
 FRHITransientTexture* FRHITransientResourcePageAllocator::CreateTexture(
 	const FRHITextureCreateInfo& CreateInfo,
 	const TCHAR* DebugName,
@@ -1661,7 +1560,6 @@ void FRHITransientResourcePageAllocator::Flush(FRHICommandListImmediate& RHICmdL
 
 		Textures.Forfeit(GFrameCounterRenderThread, [this](FRHITransientTexture* Texture) { ReleaseTextureInternal(Texture); });
 
-		LogLongLivedResources(Textures.GetAllocated(), CurrentCycle);
 		TRACE_COUNTER_SET(TransientTextureCacheSize, Textures.GetSize());
 	}
 
@@ -1671,7 +1569,6 @@ void FRHITransientResourcePageAllocator::Flush(FRHICommandListImmediate& RHICmdL
 
 		Buffers.Forfeit(GFrameCounterRenderThread, [this](FRHITransientBuffer* Buffer) { ReleaseBufferInternal(Buffer); });
 
-		LogLongLivedResources(Buffers.GetAllocated(), CurrentCycle);
 		TRACE_COUNTER_SET(TransientBufferCacheSize, Buffers.GetSize());
 	}
 
@@ -1700,4 +1597,4 @@ void FRHITransientResourcePageAllocator::Flush(FRHICommandListImmediate& RHICmdL
 	Stats.Reset();
 
 	CurrentCycle++;
-}
+}
\ No newline at end of file
diff --git a/Engine/Source/Runtime/RHICore/Public/RHICoreTransientResourceAllocator.h b/Engine/Source/Runtime/RHICore/Public/RHICoreTransientResourceAllocator.h
--- a/Engine/Source/Runtime/RHICore/Public/RHICoreTransientResourceAllocator.h
+++ b/Engine/Source/Runtime/RHICore/Public/RHICoreTransientResourceAllocator.h
@@ -327,7 +327,6 @@ class RHICORE_API FRHITransientHeapAllocator
 {
 public:
 	FRHITransientHeapAllocator(uint64 Capacity, uint32 Alignment);
-	~FRHITransientHeapAllocator();
 
 	FRHITransientHeapAllocation Allocate(uint64 Size, uint32 Alignment);
 
@@ -661,8 +660,6 @@ public:
 		: HeapCache(InHeapCache)
 	{}
 
-	~FRHITransientResourceHeapAllocator();
-
 	// Deallocates a texture from its parent heap. Provide the current platform fence value used to update the heap.
 	void DeallocateMemory(FRHITransientTexture* Texture, uint32 PassIndex) override;
 
@@ -1056,8 +1053,6 @@ public:
 		FastPagePool = PagePoolCache.GetFastPagePool();
 	}
 
-	~FRHITransientResourcePageAllocator();
-
 	FRHITransientTexture* CreateTexture(const FRHITextureCreateInfo& CreateInfo, const TCHAR* DebugName, uint32 PassIndex) override;
 	FRHITransientBuffer* CreateBuffer(const FRHIBufferCreateInfo& CreateInfo, const TCHAR* DebugName, uint32 PassIndex) override;
 	void DeallocateMemory(FRHITransientTexture* Texture, uint32 PassIndex) override;
diff --git a/Engine/Source/Runtime/RenderCore/Private/RenderGraphResourcePool.cpp b/Engine/Source/Runtime/RenderCore/Private/RenderGraphResourcePool.cpp
--- a/Engine/Source/Runtime/RenderCore/Private/RenderGraphResourcePool.cpp
+++ b/Engine/Source/Runtime/RenderCore/Private/RenderGraphResourcePool.cpp
@@ -155,7 +155,14 @@ uint32 FRDGTransientRenderTarget::Release()
 	const uint32 Refs = --RefCount;
 	if (Refs == 0)
 	{
-		GRDGTransientResourceAllocator.AddPendingDeallocation(this);
+		if (GRDGTransientResourceAllocator.IsValid())
+		{
+			GRDGTransientResourceAllocator.AddPendingDeallocation(this);
+		}
+		else
+		{
+			delete this;
+		}
 	}
 	return Refs;
 }
diff --git a/Engine/Source/Runtime/RenderCore/Private/RenderGraphResourcePool.h b/Engine/Source/Runtime/RenderCore/Private/RenderGraphResourcePool.h
--- a/Engine/Source/Runtime/RenderCore/Private/RenderGraphResourcePool.h
+++ b/Engine/Source/Runtime/RenderCore/Private/RenderGraphResourcePool.h
@@ -100,6 +100,8 @@ public:
 
 	void ReleasePendingDeallocations();
 
+	bool IsValid() const { return Allocator != nullptr; }
+
 private:
 	void InitDynamicRHI() override;
 	void ReleaseDynamicRHI() override;
diff --git a/Engine/Source/Runtime/Renderer/Private/SceneTextures.cpp b/Engine/Source/Runtime/Renderer/Private/SceneTextures.cpp
--- a/Engine/Source/Runtime/Renderer/Private/SceneTextures.cpp
+++ b/Engine/Source/Runtime/Renderer/Private/SceneTextures.cpp
@@ -903,7 +903,7 @@ uint32 FSceneTextures::GetGBufferRenderTargets(ERenderTargetLoadAction LoadActio
 void FSceneTextureExtracts::QueueExtractions(FRDGBuilder& GraphBuilder, const FSceneTextures& SceneTextures)
 {
 	// Free up the memory for reuse during the RDG execution phase.
-	ReleaseDynamicRHI();
+	Release();
 
 	ESceneTextureSetupMode SetupMode = ESceneTextureSetupMode::None;
 
@@ -952,7 +952,7 @@ void FSceneTextureExtracts::QueueExtractions(FRDGBuilder& GraphBuilder, const FS
 	});
 }
 
-void FSceneTextureExtracts::ReleaseDynamicRHI()
+void FSceneTextureExtracts::Release()
 {
 	Depth = {};
 	CustomDepth = {};
diff --git a/Engine/Source/Runtime/Renderer/Public/SceneRenderTargetParameters.h b/Engine/Source/Runtime/Renderer/Public/SceneRenderTargetParameters.h
--- a/Engine/Source/Runtime/Renderer/Public/SceneRenderTargetParameters.h
+++ b/Engine/Source/Runtime/Renderer/Public/SceneRenderTargetParameters.h
@@ -170,7 +170,7 @@ public:
 
 private:
 	void Release();
-	void ReleaseDynamicRHI() override;
+	void ReleaseDynamicRHI() override { Release(); }
 
 	// Contains the resolved scene depth target.
 	TRefCountPtr<IPooledRenderTarget> Depth;
