Fixed alignment issues with D3D12 transient allocator.

diff --git a/Engine/Source/Runtime/D3D12RHI/Private/D3D12TransientResourceAllocator.cpp b/Engine/Source/Runtime/D3D12RHI/Private/D3D12TransientResourceAllocator.cpp
--- a/Engine/Source/Runtime/D3D12RHI/Private/D3D12TransientResourceAllocator.cpp
+++ b/Engine/Source/Runtime/D3D12RHI/Private/D3D12TransientResourceAllocator.cpp
@@ -54,7 +54,7 @@ FD3D12TransientHeap::FD3D12TransientHeap(const FInitializer& Initializer, FD3D12
 	D3D12_HEAP_DESC Desc = {};
 	Desc.SizeInBytes = Initializer.Size;
 	Desc.Properties = HeapProperties;
-	Desc.Alignment = D3D12_DEFAULT_RESOURCE_PLACEMENT_ALIGNMENT;
+	Desc.Alignment = D3D12_DEFAULT_MSAA_RESOURCE_PLACEMENT_ALIGNMENT;
 	Desc.Flags = HeapFlags;
 
 	if (Adapter->IsHeapNotZeroedSupported())
@@ -86,7 +86,7 @@ FD3D12TransientHeap::FD3D12TransientHeap(const FInitializer& Initializer, FD3D12
 	Heap->SetHeap(D3DHeap, TEXT("TransientResourceAllocator Backing Heap"), true, true);
 	Heap->BeginTrackingResidency(Desc.SizeInBytes);
 
-	SetGPUVirtualAddress(Heap->GetGPUVirtualAddress());
+	SetGpuVirtualAddress(Heap->GetGPUVirtualAddress());
 
 	INC_MEMORY_STAT_BY(STAT_D3D12TransientHeaps, Desc.SizeInBytes);
 	INC_MEMORY_STAT_BY(STAT_D3D12MemoryCurrentTotal, Desc.SizeInBytes);
diff --git a/Engine/Source/Runtime/RHICore/Private/RHICoreTransientResourceAllocator.cpp b/Engine/Source/Runtime/RHICore/Private/RHICoreTransientResourceAllocator.cpp
--- a/Engine/Source/Runtime/RHICore/Private/RHICoreTransientResourceAllocator.cpp
+++ b/Engine/Source/Runtime/RHICore/Private/RHICoreTransientResourceAllocator.cpp
@@ -370,7 +370,7 @@ FRHITransientHeapAllocator::FFindResult FRHITransientHeapAllocator::FindFreeRang
 		FRange& Range = Ranges[Handle];
 
 		// Due to alignment we may have to shift the offset and expand the size accordingly.
-		const uint64 AlignmentPad = Align(Range.Offset, Alignment) - Range.Offset;
+		const uint64 AlignmentPad = Align(GpuVirtualAddress + Range.Offset, Alignment) - GpuVirtualAddress - Range.Offset;
 		const uint64 RequiredSize = Size + AlignmentPad;
 
 		if (RequiredSize <= Range.Size)
diff --git a/Engine/Source/Runtime/RHICore/Public/RHICoreTransientResourceAllocator.h b/Engine/Source/Runtime/RHICore/Public/RHICoreTransientResourceAllocator.h
--- a/Engine/Source/Runtime/RHICore/Public/RHICoreTransientResourceAllocator.h
+++ b/Engine/Source/Runtime/RHICore/Public/RHICoreTransientResourceAllocator.h
@@ -332,6 +332,12 @@ public:
 
 	void Deallocate(FRHITransientHeapAllocation Allocation);
 
+	void SetGpuVirtualAddress(uint64 InGpuVirtualAddress)
+	{
+		GpuVirtualAddress = InGpuVirtualAddress;
+	}
+
+	uint64 GetGpuVirtualAddress() const { return GpuVirtualAddress; }
 	uint64 GetCapacity() const { return Capacity; }
 	uint64 GetUsedSize() const { return UsedSize; }
 	uint64 GetFreeSize() const { return Capacity - UsedSize; }
@@ -407,6 +413,7 @@ private:
 
 	void Validate();
 
+	uint64 GpuVirtualAddress = 0;
 	uint64 Capacity = 0;
 	uint64 UsedSize = 0;
 	uint64 AlignmentWaste = 0;
@@ -523,7 +530,7 @@ public:
 
 	uint64 GetCapacity() const { return Allocator.GetCapacity(); }
 
-	uint64 GetGPUVirtualAddress() const { return GpuVirtualAddress; }
+	uint64 GetGPUVirtualAddress() const { return Allocator.GetGpuVirtualAddress(); }
 
 	uint64 GetLastUsedGarbageCollectCycle() const { return LastUsedGarbageCollectCycle; }
 
@@ -541,9 +548,9 @@ public:
 	}
 
 protected:
-	void SetGPUVirtualAddress(uint64 InBaseGPUVirtualAddress)
+	void SetGpuVirtualAddress(uint64 InBaseGPUVirtualAddress)
 	{
-		GpuVirtualAddress = InBaseGPUVirtualAddress;
+		Allocator.SetGpuVirtualAddress(InBaseGPUVirtualAddress);
 	}
 
 private:
@@ -553,7 +560,6 @@ private:
 	FInitializer Initializer;
 	FRHITransientHeapAllocator Allocator;
 
-	uint64 GpuVirtualAddress = 0;
 	uint64 LastUsedGarbageCollectCycle = 0;
 	uint64 CommitSize = 0;
 	uint32 AlignmentLog2;
@@ -864,10 +870,9 @@ public:
 			, Spans(InResource.GetPageAllocation().Spans)
 			, AllocationsBefore(Allocations)
 			, GpuVirtualAddress(Resource.GetGpuVirtualAddress())
-			, Size(Resource.GetSize())
+			, Size(Align(Resource.GetSize(), InPageSize))
 			, PagesRemaining(Size / InPageSize)
 		{
-			checkf(Size % InPageSize == 0, TEXT("Resource size was not aligned properly to page size."));
 			Allocations.Reset();
 			Spans.Reset();
 		}
