Add fast VRAM percentage control to transient resource for hinting alloctor to allocated a precise amount.

diff --git a/Engine/Source/Runtime/RHI/Public/RHIResources.h b/Engine/Source/Runtime/RHI/Public/RHIResources.h
--- a/Engine/Source/Runtime/RHI/Public/RHIResources.h
+++ b/Engine/Source/Runtime/RHI/Public/RHIResources.h
@@ -1496,6 +1496,9 @@ struct RHI_API FRHITextureDesc
 	/** Texture format used when creating the UAV. PF_Unknown means to use the default one (same as Format). */
 	EPixelFormat UAVFormat = PF_Unknown;
 
+	/** Resource memory percentage which should be allocated onto fast VRAM (hint-only). (encoding into 8bits, 0..255 -> 0%..100%) */
+	uint8 FastVRAMPercentage = 0xFF;
+
 	/** Check the validity. */
 	static bool CheckValidity(const FRHITextureDesc& Desc, const TCHAR* Name)
 	{
@@ -1669,7 +1672,8 @@ struct FRHITextureCreateDesc : public FRHITextureDesc
 	FRHITextureCreateDesc& SetDebugName(const TCHAR* InDebugName)              { DebugName = InDebugName;                  return *this; }
 	FRHITextureCreateDesc& SetGPUMask(FRHIGPUMask InGPUMask)                   { GPUMask = InGPUMask;                      return *this; }
 	FRHITextureCreateDesc& SetBulkData(FResourceBulkDataInterface* InBulkData) { BulkData = InBulkData;                    return *this; }
-	FRHITextureCreateDesc& DetermineInititialState() { if (InitialState == ERHIAccess::Unknown) InitialState = RHIGetDefaultResourceState(Flags, BulkData != nullptr); return *this; }
+	FRHITextureCreateDesc& DetermineInititialState()                           { if (InitialState == ERHIAccess::Unknown) InitialState = RHIGetDefaultResourceState(Flags, BulkData != nullptr); return *this; }
+	FRHITextureCreateDesc& SetFastVRAMPercentage(float In)                     { FastVRAMPercentage = FMath::Clamp(In, 0.f, 1.0f) * 0xFF; return *this; }
 
 	/* The RHI access state that the resource will be created in. */
 	ERHIAccess InitialState = ERHIAccess::Unknown;
diff --git a/Engine/Source/Runtime/RHICore/Private/RHICoreTransientResourceAllocator.cpp b/Engine/Source/Runtime/RHICore/Private/RHICoreTransientResourceAllocator.cpp
--- a/Engine/Source/Runtime/RHICore/Private/RHICoreTransientResourceAllocator.cpp
+++ b/Engine/Source/Runtime/RHICore/Private/RHICoreTransientResourceAllocator.cpp
@@ -1194,7 +1194,9 @@ void FRHITransientPagePool::Allocate(FAllocationContext& Context)
 {
 	uint32 SpanIndex = 0;
 	uint32 PagesAllocated = 0;
-	if (Allocator.Allocate(Context.PagesRemaining, PagesAllocated, SpanIndex))
+
+	uint32 PagesRemaining = Context.MaxAllocationPage > 0 ? FMath::Min(Context.PagesRemaining, Context.MaxAllocationPage) : Context.PagesRemaining;
+	if (Allocator.Allocate(PagesRemaining, PagesAllocated, SpanIndex))
 	{
 		const uint64 DestinationGpuVirtualAddress = Context.GpuVirtualAddress + Context.PagesAllocated * Initializer.PageSize;
 		const uint32 PageSpanOffsetMin = PageSpans.Num();
@@ -1405,9 +1407,12 @@ FRHITransientTexture* FRHITransientResourcePageAllocator::CreateTexture(
 		return CreateTextureInternal(CreateInfo, DebugName, Hash);
 	});
 
+	const bool bFastPool = EnumHasAnyFlags(CreateInfo.Flags, ETextureCreateFlags::FastVRAM) || EnumHasAnyFlags(CreateInfo.Flags, ETextureCreateFlags::FastVRAMPartialAlloc);
+	const float FastPoolPercentageRequested = bFastPool ? CreateInfo.FastVRAMPercentage / 255.f : 0.f;
+
 	check(Texture);
 	Texture->Acquire(DebugName, PassIndex, CurrentCycle);
-	AllocateMemoryInternal(Texture, DebugName, PassIndex, EnumHasAnyFlags(CreateInfo.Flags, ETextureCreateFlags::FastVRAM));
+	AllocateMemoryInternal(Texture, DebugName, PassIndex, bFastPool, FastPoolPercentageRequested);
 	Stats.AllocateTexture(Texture->GetSize());
 	IF_RHICORE_TRANSIENT_ALLOCATOR_DEBUG(ActiveResources.Emplace(Texture));
 	return Texture;
@@ -1426,19 +1431,22 @@ FRHITransientBuffer* FRHITransientResourcePageAllocator::CreateBuffer(
 
 	check(Buffer);
 	Buffer->Acquire(DebugName, PassIndex, CurrentCycle);
-	AllocateMemoryInternal(Buffer, DebugName, PassIndex, EnumHasAnyFlags(CreateInfo.Usage, EBufferUsageFlags::FastVRAM));
+	AllocateMemoryInternal(Buffer, DebugName, PassIndex, EnumHasAnyFlags(CreateInfo.Usage, EBufferUsageFlags::FastVRAM), false);
 	Stats.AllocateBuffer(Buffer->GetSize());
 	IF_RHICORE_TRANSIENT_ALLOCATOR_DEBUG(ActiveResources.Emplace(Buffer));
 	return Buffer;
 }
 
-void FRHITransientResourcePageAllocator::AllocateMemoryInternal(FRHITransientResource* Resource, const TCHAR* DebugName, uint32 PassIndex, bool bFastPoolRequested)
+void FRHITransientResourcePageAllocator::AllocateMemoryInternal(FRHITransientResource* Resource, const TCHAR* DebugName, uint32 PassIndex, bool bFastPoolRequested, float FastPoolPercentageRequested)
 {
 	FRHITransientPagePool::FAllocationContext AllocationContext(*Resource, PageSize);
 
 	if (bFastPoolRequested && FastPagePool)
 	{
+		// If a partial allocation is requested, compute the max. number of page which should be allocated in fast memory
+		AllocationContext.MaxAllocationPage = FastPoolPercentageRequested > 0 ? FMath::CeilToInt(AllocationContext.PagesRemaining * FastPoolPercentageRequested) : AllocationContext.MaxAllocationPage;
 		FastPagePool->Allocate(AllocationContext);
+		AllocationContext.MaxAllocationPage = 0;
 	}
 
 	if (!AllocationContext.IsComplete())
diff --git a/Engine/Source/Runtime/RHICore/Public/RHICoreTransientResourceAllocator.h b/Engine/Source/Runtime/RHICore/Public/RHICoreTransientResourceAllocator.h
--- a/Engine/Source/Runtime/RHICore/Public/RHICoreTransientResourceAllocator.h
+++ b/Engine/Source/Runtime/RHICore/Public/RHICoreTransientResourceAllocator.h
@@ -889,6 +889,7 @@ public:
 		uint32 AllocationCount = 0;
 		uint32 AllocationMatchingCount = 0;
 		uint32 PagesRemaining = 0;
+		uint32 MaxAllocationPage = 0;
 		uint32 PagesAllocated = 0;
 		uint32 PageSpansAllocated = 0;
 		uint32 PagesMapped = 0;
@@ -1075,7 +1076,7 @@ private:
 	static constexpr uint64 KB = 1024;
 	static constexpr uint64 MB = 1024 * KB;
 
-	void AllocateMemoryInternal(FRHITransientResource* Resource, const TCHAR* DebugName, uint32 PassIndex, bool bFastPoolRequested);
+	void AllocateMemoryInternal(FRHITransientResource* Resource, const TCHAR* DebugName, uint32 PassIndex, bool bFastPoolRequested, float FastPoolPercentageRequested);
 	void DeallocateMemoryInternal(FRHITransientResource* Resource, uint32 PassIndex);
 
 	//////////////////////////////////////////////////////////////////////////
