- Fix possible wrong allocation size when pool alignment isn't a multiple of the allocation alignment - it could allocate too little
- Move linked list modification to protected code

diff --git a/Engine/Source/Runtime/RHICore/Private/RHIPoolAllocator.cpp b/Engine/Source/Runtime/RHICore/Private/RHIPoolAllocator.cpp
--- a/Engine/Source/Runtime/RHICore/Private/RHIPoolAllocator.cpp
+++ b/Engine/Source/Runtime/RHICore/Private/RHIPoolAllocator.cpp
@@ -64,20 +64,20 @@ void FRHIPoolAllocationData::InitAsFree(int16 InPoolIndex, uint32 InSize, uint32
 }
 
 
-void FRHIPoolAllocationData::InitAsAllocated(uint32 InSize, uint32 InAlignment, FRHIPoolAllocationData* InFree)
+void FRHIPoolAllocationData::InitAsAllocated(uint32 InSize, uint32 InPoolAlignment, uint32 InAllocationAlignment, FRHIPoolAllocationData* InFree)
 {
 	check(InFree->IsFree());
 
 	Reset();
 	
 	Size = InSize;
-	Alignment = InAlignment;
+	Alignment = InAllocationAlignment;
 	SetAllocationType(EAllocationType::Allocated);
-	Offset = FRHIMemoryPool::GetAlignedOffset(InFree->Offset, InFree->Alignment, InAlignment);
+	Offset = FRHIMemoryPool::GetAlignedOffset(InFree->Offset, InPoolAlignment, InAllocationAlignment);
 	PoolIndex = InFree->PoolIndex;
 	Locked = true;
 
-	uint32 AlignedSize = AlignArbitrary(InSize, InFree->Alignment);
+	uint32 AlignedSize = FRHIMemoryPool::GetAlignedSize(InSize, InPoolAlignment, InAllocationAlignment);
 	InFree->Size -= AlignedSize;
 	InFree->Offset += AlignedSize;
 	InFree->AddBefore(this);
@@ -129,15 +129,15 @@ void FRHIPoolAllocationData::MoveFrom(FRHIPoolAllocationData& InAllocated, bool
 }
 
 
-void FRHIPoolAllocationData::MarkFree(uint32 InAlignment)
+void FRHIPoolAllocationData::MarkFree(uint32 InPoolAlignment, uint32 InAllocationAlignment)
 {
 	check(IsAllocated());
 
 	// Mark as free and update the size, offset and alignment according to the new requested alignment
 	SetAllocationType(EAllocationType::Free);
-	Size = Align(Size, InAlignment);
-	Offset = AlignDown(Offset, InAlignment);
-	Alignment = InAlignment;
+	Size = FRHIMemoryPool::GetAlignedSize(Size, InPoolAlignment, InAllocationAlignment);
+	Offset = AlignDown(Offset, InPoolAlignment);
+	Alignment = InPoolAlignment;
 }
 
 
@@ -351,7 +351,7 @@ bool FRHIMemoryPool::TryAllocate(uint32 InSizeInBytes, uint32 InAllocationAlignm
 		FreeBlocks.RemoveAt(FreeBlockIndex);
 
 		// update private allocator data of new and free block
-		AllocationData.InitAsAllocated(InSizeInBytes, InAllocationAlignment, FreeBlock);
+		AllocationData.InitAsAllocated(InSizeInBytes, PoolAlignment, InAllocationAlignment, FreeBlock);
 		check((AllocationData.GetOffset() % InAllocationAlignment) == 0);
 
 		// Update working stats
@@ -387,13 +387,16 @@ void FRHIMemoryPool::Deallocate(FRHIPoolAllocationData& AllocationData)
 	check(AllocationData.IsLocked());
 	check(PoolIndex == AllocationData.GetPoolIndex());
 
+	// Original allocation is used to compute the orignal aligned allocation size
+	uint32 AllocationAlignment = AllocationData.GetAlignment();
+
 	// Free block should not be locked anymore - can be reused immediatly when we get to actual pool deallocate
 	bool bLocked = false;
 	uint64 AllocationSize = AllocationData.GetSize();
 
 	FRHIPoolAllocationData* FreeBlock = GetNewAllocationData();
 	FreeBlock->MoveFrom(AllocationData, bLocked);
-	FreeBlock->MarkFree(PoolAlignment);
+	FreeBlock->MarkFree(PoolAlignment, AllocationAlignment);
 		
 	// Update working stats
 	FreeSize += FreeBlock->GetSize();
@@ -584,6 +587,10 @@ void FRHIMemoryPool::Validate()
 	FRHIPoolAllocationData* CurrentBlock = HeadBlock.GetNext();
 	while (CurrentBlock != &HeadBlock)
 	{
+		// validate linked list
+		check(CurrentBlock == CurrentBlock->GetNext()->GetPrev());
+		check(CurrentBlock == CurrentBlock->GetPrev()->GetNext());
+
 		uint32 AlignedSize = GetAlignedSize(CurrentBlock->GetSize(), PoolAlignment, CurrentBlock->GetAlignment());
 		uint32 AlignedOffset = AlignDown(CurrentBlock->GetOffset(), PoolAlignment);
 
diff --git a/Engine/Source/Runtime/RHICore/Public/RHIPoolAllocator.h b/Engine/Source/Runtime/RHICore/Public/RHIPoolAllocator.h
--- a/Engine/Source/Runtime/RHICore/Public/RHIPoolAllocator.h
+++ b/Engine/Source/Runtime/RHICore/Public/RHIPoolAllocator.h
@@ -40,19 +40,13 @@ public:
 	// Initialize options
 	void InitAsHead(int16 InPoolIndex);
 	void InitAsFree(int16 InPoolIndex, uint32 InSize, uint32 InAlignment, uint32 InOffset);
-	void InitAsAllocated(uint32 InSize, uint32 InAlignment, FRHIPoolAllocationData* InFree);
+	void InitAsAllocated(uint32 InSize, uint32 InPoolAlignment, uint32 InAllocationAlignment, FRHIPoolAllocationData* InFree);
 	void MoveFrom(FRHIPoolAllocationData& InAllocated, bool InLocked);
 
-	// Free block operation
-	void MarkFree(uint32 InAlignment);
-	void Merge(FRHIPoolAllocationData* InOther);
+	// Free block operation (TODO: make internal)
+	void MarkFree(uint32 InPoolAlignment, uint32 InAllocationAlignment);
 
-	// Linked list operation
-	void RemoveFromLinkedList();
-	void AddBefore(FRHIPoolAllocationData* InOther);
-	void AddAfter(FRHIPoolAllocationData* InOther);
-
-	// Alias operations
+	// Alias operations (NOTE: currently not accessed when pool lock is taken, needs to be fixed)
 	void AddAlias(FRHIPoolAllocationData* InOther);
 	void RemoveAlias();
 	FRHIPoolAllocationData* GetFirstAlias() const { return AliasAllocation; }
@@ -87,6 +81,13 @@ private:
 	friend class FRHIPoolAllocator;
 	void Lock() { check(Locked == 0); Locked = 1; }
 
+	void Merge(FRHIPoolAllocationData* InOther);
+
+	// Linked list operation
+	void RemoveFromLinkedList();
+	void AddBefore(FRHIPoolAllocationData* InOther);
+	void AddAfter(FRHIPoolAllocationData* InOther);
+		
 	enum class EAllocationType : uint8
 	{
 		Unknown,
