Fix FRHITransientResourceOverlapTracker forgetting about memory region while tracking resources.

diff --git a/Engine/Source/Runtime/RHICore/Private/RHICoreTransientResourceAllocator.cpp b/Engine/Source/Runtime/RHICore/Private/RHICoreTransientResourceAllocator.cpp
--- a/Engine/Source/Runtime/RHICore/Private/RHICoreTransientResourceAllocator.cpp
+++ b/Engine/Source/Runtime/RHICore/Private/RHICoreTransientResourceAllocator.cpp
@@ -146,40 +146,35 @@ void FRHITransientResourceOverlapTracker::Track(FRHITransientResource* Transient
 		// Partial overlap, can manifest as three cases:
 		else
 		{
-			bool bResizedOld = false;
-			FResourceRange ResourceRangeOldCopy = ResourceRangeOld;
-
 			// 1) New:    ********
 			//            |||        ->
 			//    Old: ======             ===********
 			if (ResourceRangeOld.PageOffsetMin < ResourceRangeNew.PageOffsetMin)
 			{
+				FResourceRange ResourceRangeOldCopy = ResourceRangeOld;
 				ResourceRangeOld.PageOffsetMax = ResourceRangeNew.PageOffsetMin;
-				bResizedOld = true;
 
-				// This case should only ever happen on the first iteration.
-				check(Index == InsertIndex);
 				InsertIndex++;
-			}
-			else
-			{
-				// 2) New:    ********
-				//                |||      ->
-				//    Old:        ======         ********===
-				if (!bResizedOld)
-				{
-					ResourceRangeOld.PageOffsetMin = ResourceRangeNew.PageOffsetMax;
-				}
-
 				// 3) New:    ********
 				//            ||||||||      ->
 				//    Old: ==============        ===********===
-				else
+				if (ResourceRangeOldCopy.PageOffsetMax > ResourceRangeNew.PageOffsetMax)
 				{
+					ResourceRangeOldCopy.PageOffsetMin = ResourceRangeNew.PageOffsetMax;
+
 					// Lower bound has been resized already; add an upper bound.
 					FResourceRange ResourceRangeOldUpper = ResourceRangeOldCopy;
-					Index = ResourceRanges.Insert(ResourceRangeOldUpper, Index + 1);
+					ResourceRanges.Insert(ResourceRangeOldUpper, Index + 1);
+
+					break;
 				}
+			}
+			else
+			{
+				// 2) New:    ********
+				//                |||      ->
+				//    Old:        ======         ********===
+				ResourceRangeOld.PageOffsetMin = ResourceRangeNew.PageOffsetMax;
 
 				break;
 			}
@@ -1605,4 +1600,4 @@ void FRHITransientResourcePageAllocator::Flush(FRHICommandListImmediate& RHICmdL
 	Stats.Reset();
 
 	CurrentCycle++;
-}
\ No newline at end of file
+}
