RHI: Add break even if not necessary to match rest of code.

diff --git a/Engine/Source/Runtime/RHICore/Internal/RHIUniformBufferDataShared.h b/Engine/Source/Runtime/RHICore/Internal/RHIUniformBufferDataShared.h
--- a/Engine/Source/Runtime/RHICore/Internal/RHIUniformBufferDataShared.h
+++ b/Engine/Source/Runtime/RHICore/Internal/RHIUniformBufferDataShared.h
@@ -65,6 +65,7 @@ namespace UE::RHICore
 			FRHIUnorderedAccessView* UnorderedAccessView = Reader.Read<FRHIUnorderedAccessView*>(Resource);
 			return UnorderedAccessView ? UnorderedAccessView->GetBindlessHandle() : FRHIDescriptorHandle();
 		}
+		break;
 		case UBMT_RDG_TEXTURE_UAV:
 		case UBMT_RDG_BUFFER_UAV:
 		{
