Fixing an off-by-one issue in the descriptor allocator.

diff --git a/Engine/Source/Runtime/RHICore/Private/RHIDescriptorAllocator.cpp b/Engine/Source/Runtime/RHICore/Private/RHIDescriptorAllocator.cpp
--- a/Engine/Source/Runtime/RHICore/Private/RHIDescriptorAllocator.cpp
+++ b/Engine/Source/Runtime/RHICore/Private/RHIDescriptorAllocator.cpp
@@ -149,7 +149,7 @@ void FRHIDescriptorAllocator::Free(uint32 Offset, uint32 NumDescriptors)
 				{
 					// Merge with next range
 					Ranges[Index].Last = Ranges[Index + 1].Last;
-					Ranges.RemoveAt(Index);
+					Ranges.RemoveAt(Index + 1);
 				}
 				else
 				{
