Descriptor Management Stats
- Adding stats to track general descriptors as well as bindless descriptors
- Adding descriptor heap manager that's offset from the start of a descriptor heap.

diff --git a/Engine/Source/Runtime/D3D12RHI/Private/D3D12Descriptors.cpp b/Engine/Source/Runtime/D3D12RHI/Private/D3D12Descriptors.cpp
--- a/Engine/Source/Runtime/D3D12RHI/Private/D3D12Descriptors.cpp
+++ b/Engine/Source/Runtime/D3D12RHI/Private/D3D12Descriptors.cpp
@@ -31,6 +31,17 @@ FD3D12DescriptorHeap::FD3D12DescriptorHeap(FD3D12Device* InDevice, ID3D12Descrip
 	, bIsGlobal(bInIsGlobal)
 	, bIsSuballocation(false)
 {
+#if STATS
+	const size_t MemorySize = DescriptorSize * NumDescriptors;
+	if (Type == ERHIDescriptorHeapType::Standard)
+	{
+		INC_MEMORY_STAT_BY(STAT_BindlessResourceHeapMemory, MemorySize);
+	}
+	else if (Type == ERHIDescriptorHeapType::Sampler)
+	{
+		INC_MEMORY_STAT_BY(STAT_BindlessSamplerHeapMemory, MemorySize);
+	}
+#endif
 }
 
 FD3D12DescriptorHeap::FD3D12DescriptorHeap(FD3D12DescriptorHeap* SubAllocateSourceHeap, uint32 InOffset, uint32 InNumDescriptors)
@@ -48,14 +59,30 @@ FD3D12DescriptorHeap::FD3D12DescriptorHeap(FD3D12DescriptorHeap* SubAllocateSour
 {
 }
 
-FD3D12DescriptorHeap::~FD3D12DescriptorHeap() = default;
+FD3D12DescriptorHeap::~FD3D12DescriptorHeap()
+{
+#if STATS
+	if (!bIsSuballocation)
+	{
+		const size_t MemorySize = DescriptorSize * NumDescriptors;
+		if (Type == ERHIDescriptorHeapType::Standard)
+		{
+			DEC_MEMORY_STAT_BY(STAT_BindlessResourceHeapMemory, MemorySize);
+		}
+		else if (Type == ERHIDescriptorHeapType::Sampler)
+		{
+			DEC_MEMORY_STAT_BY(STAT_BindlessSamplerHeapMemory, MemorySize);
+		}
+	}
+#endif
+}
 
 ////////////////////////////////////////////////////////////////////////////////////////////////////////////////////////
 // FD3D12DescriptorManager
 
-FD3D12DescriptorManager::FD3D12DescriptorManager(FD3D12Device* Device, FD3D12DescriptorHeap* InHeap)
+FD3D12DescriptorManager::FD3D12DescriptorManager(FD3D12Device* Device, FD3D12DescriptorHeap* InHeap, TConstArrayView<TStatId> InStats)
 	: FD3D12DeviceChild(Device)
-	, FRHIHeapDescriptorAllocator(InHeap->GetType(), InHeap->GetNumDescriptors())
+	, FRHIHeapDescriptorAllocator(InHeap->GetType(), InHeap->GetNumDescriptors(), InStats)
 	, Heap(InHeap)
 {
 }
@@ -95,7 +122,14 @@ void FD3D12BindlessDescriptorManager::Init(uint32 InNumResourceDescriptors, uint
 			InNumResourceDescriptors,
 			D3D12_DESCRIPTOR_HEAP_FLAG_SHADER_VISIBLE
 		);
-		Managers.Emplace(GetParentDevice(), Heap);
+
+		const TStatId Stats[] =
+		{
+			GET_STATID(STAT_ResourceDescriptorsAllocated),
+			GET_STATID(STAT_BindlessResourceDescriptorsAllocated),
+		};
+
+		Managers.Emplace(GetParentDevice(), Heap, Stats);
 	}
 
 	if (InNumSamplerDescriptors > 0)
@@ -106,7 +140,14 @@ void FD3D12BindlessDescriptorManager::Init(uint32 InNumResourceDescriptors, uint
 			InNumSamplerDescriptors,
 			D3D12_DESCRIPTOR_HEAP_FLAG_SHADER_VISIBLE
 		);
-		Managers.Emplace(GetParentDevice(), Heap);
+
+		const TStatId Stats[] =
+		{
+			GET_STATID(STAT_SamplerDescriptorsAllocated),
+			GET_STATID(STAT_BindlessSamplerDescriptorsAllocated),
+		};
+
+		Managers.Emplace(GetParentDevice(), Heap, Stats);
 	}
 }
 
@@ -456,7 +497,8 @@ void FD3D12DescriptorHeapManager::Init(uint32 InNumGlobalResourceDescriptors, ui
 			true /* bIsGlobal */
 		);
 
-		GlobalHeaps.Emplace(GetParentDevice(), Heap);
+		const TStatId Stats[] = {GET_STATID(STAT_ResourceDescriptorsAllocated)};
+		GlobalHeaps.Emplace(GetParentDevice(), Heap, Stats);
 	}
 
 	if (InNumGlobalSamplerDescriptors > 0)
@@ -470,7 +512,8 @@ void FD3D12DescriptorHeapManager::Init(uint32 InNumGlobalResourceDescriptors, ui
 			true /* bIsGlobal */
 		);
 
-		GlobalHeaps.Emplace(GetParentDevice(), Heap);
+		const TStatId Stats[] = { GET_STATID(STAT_SamplerDescriptorsAllocated) };
+		GlobalHeaps.Emplace(GetParentDevice(), Heap, Stats);
 	}
 }
 
diff --git a/Engine/Source/Runtime/D3D12RHI/Private/D3D12Descriptors.h b/Engine/Source/Runtime/D3D12RHI/Private/D3D12Descriptors.h
--- a/Engine/Source/Runtime/D3D12RHI/Private/D3D12Descriptors.h
+++ b/Engine/Source/Runtime/D3D12RHI/Private/D3D12Descriptors.h
@@ -65,7 +65,7 @@ class FD3D12DescriptorManager : public FD3D12DeviceChild, public FRHIHeapDescrip
 {
 public:
 	FD3D12DescriptorManager() = delete;
-	FD3D12DescriptorManager(FD3D12Device* Device, FD3D12DescriptorHeap* InHeap);
+	FD3D12DescriptorManager(FD3D12Device* Device, FD3D12DescriptorHeap* InHeap, TConstArrayView<TStatId> InStats);
 	~FD3D12DescriptorManager();
 
 	void UpdateImmediately(FRHIDescriptorHandle InHandle, D3D12_CPU_DESCRIPTOR_HANDLE InSourceCpuHandle);
diff --git a/Engine/Source/Runtime/RHI/Private/RHIStats.cpp b/Engine/Source/Runtime/RHI/Private/RHIStats.cpp
--- a/Engine/Source/Runtime/RHI/Private/RHIStats.cpp
+++ b/Engine/Source/Runtime/RHI/Private/RHIStats.cpp
@@ -28,3 +28,11 @@ DEFINE_STAT(STAT_VertexBufferMemory);
 DEFINE_STAT(STAT_RTAccelerationStructureMemory);
 DEFINE_STAT(STAT_StructuredBufferMemory);
 DEFINE_STAT(STAT_PixelBufferMemory);
+
+DEFINE_STAT(STAT_SamplerDescriptorsAllocated);
+DEFINE_STAT(STAT_ResourceDescriptorsAllocated);
+
+DEFINE_STAT(STAT_BindlessSamplerHeapMemory);
+DEFINE_STAT(STAT_BindlessResourceHeapMemory);
+DEFINE_STAT(STAT_BindlessSamplerDescriptorsAllocated);
+DEFINE_STAT(STAT_BindlessResourceDescriptorsAllocated);
diff --git a/Engine/Source/Runtime/RHI/Public/RHIStats.h b/Engine/Source/Runtime/RHI/Public/RHIStats.h
--- a/Engine/Source/Runtime/RHI/Public/RHIStats.h
+++ b/Engine/Source/Runtime/RHI/Public/RHIStats.h
@@ -214,3 +214,12 @@ DECLARE_MEMORY_STAT_POOL_EXTERN(TEXT("Vertex buffer memory"), STAT_VertexBufferM
 DECLARE_MEMORY_STAT_POOL_EXTERN(TEXT("Ray Tracing Acceleration Structure memory"), STAT_RTAccelerationStructureMemory, STATGROUP_RHI, FPlatformMemory::MCR_GPU, RHI_API);
 DECLARE_MEMORY_STAT_POOL_EXTERN(TEXT("Structured buffer memory"), STAT_StructuredBufferMemory, STATGROUP_RHI, FPlatformMemory::MCR_GPU, RHI_API);
 DECLARE_MEMORY_STAT_POOL_EXTERN(TEXT("Pixel buffer memory"), STAT_PixelBufferMemory, STATGROUP_RHI, FPlatformMemory::MCR_GPU, RHI_API);
+
+DECLARE_DWORD_ACCUMULATOR_STAT_EXTERN(TEXT("Sampler Descriptors Allocated"), STAT_SamplerDescriptorsAllocated, STATGROUP_RHI, RHI_API);
+DECLARE_DWORD_ACCUMULATOR_STAT_EXTERN(TEXT("Resource Descriptors Allocated"), STAT_ResourceDescriptorsAllocated, STATGROUP_RHI, RHI_API);
+
+DECLARE_MEMORY_STAT_POOL_EXTERN(TEXT("Bindless Sampler Heap"), STAT_BindlessSamplerHeapMemory, STATGROUP_RHI, FPlatformMemory::MCR_GPU, RHI_API);
+DECLARE_MEMORY_STAT_POOL_EXTERN(TEXT("Bindless Resource Heap"), STAT_BindlessResourceHeapMemory, STATGROUP_RHI, FPlatformMemory::MCR_GPU, RHI_API);
+
+DECLARE_DWORD_ACCUMULATOR_STAT_EXTERN(TEXT("Bindless Sampler Descriptors Allocated"), STAT_BindlessSamplerDescriptorsAllocated, STATGROUP_RHI, RHI_API);
+DECLARE_DWORD_ACCUMULATOR_STAT_EXTERN(TEXT("Bindless Resource Descriptors Allocated"), STAT_BindlessResourceDescriptorsAllocated, STATGROUP_RHI, RHI_API);
diff --git a/Engine/Source/Runtime/RHICore/Private/RHIDescriptorAllocator.cpp b/Engine/Source/Runtime/RHICore/Private/RHIDescriptorAllocator.cpp
--- a/Engine/Source/Runtime/RHICore/Private/RHIDescriptorAllocator.cpp
+++ b/Engine/Source/Runtime/RHICore/Private/RHIDescriptorAllocator.cpp
@@ -14,19 +14,23 @@ FRHIDescriptorAllocator::FRHIDescriptorAllocator()
 {
 }
 
-FRHIDescriptorAllocator::FRHIDescriptorAllocator(uint32 InNumDescriptors)
+FRHIDescriptorAllocator::FRHIDescriptorAllocator(uint32 InNumDescriptors, TConstArrayView<TStatId> InStats)
 {
-	Init(InNumDescriptors);
+	Init(InNumDescriptors, InStats);
 }
 
 FRHIDescriptorAllocator::~FRHIDescriptorAllocator()
 {
 }
 
-void FRHIDescriptorAllocator::Init(uint32 InNumDescriptors)
+void FRHIDescriptorAllocator::Init(uint32 InNumDescriptors, TConstArrayView<TStatId> InStats)
 {
 	Capacity = InNumDescriptors;
 	Ranges.Emplace(0, InNumDescriptors - 1);
+
+#if STATS
+	Stats = InStats;
+#endif
 }
 
 void FRHIDescriptorAllocator::Shutdown()
@@ -77,6 +81,9 @@ bool FRHIDescriptorAllocator::Allocate(uint32 NumDescriptors, uint32& OutSlot)
 					CurrentRange.First += NumDescriptors;
 				}
 				OutSlot = First;
+
+				RecordAlloc(NumDescriptors);
+
 				return true;
 			}
 			++Index;
@@ -121,6 +128,8 @@ void FRHIDescriptorAllocator::Free(uint32 Offset, uint32 NumDescriptors)
 					// Just grow range
 					Ranges[Index].First = Offset;
 				}
+
+				RecordFree(NumDescriptors);
 				return;
 			}
 			else
@@ -135,6 +144,8 @@ void FRHIDescriptorAllocator::Free(uint32 Offset, uint32 NumDescriptors)
 				{
 					// Found our position in the list, insert the deleted range here
 					Ranges.EmplaceAt(Index, Offset, End - 1);
+
+					RecordFree(NumDescriptors);
 					return;
 				}
 			}
@@ -156,6 +167,8 @@ void FRHIDescriptorAllocator::Free(uint32 Offset, uint32 NumDescriptors)
 					// Just grow range
 					Ranges[Index].Last += NumDescriptors;
 				}
+
+				RecordFree(NumDescriptors);
 				return;
 			}
 			else
@@ -170,6 +183,8 @@ void FRHIDescriptorAllocator::Free(uint32 Offset, uint32 NumDescriptors)
 				{
 					// Found our position in the list, insert the deleted range here
 					Ranges.EmplaceAt(Index + 1, Offset, End - 1);
+
+					RecordFree(NumDescriptors);
 					return;
 				}
 			}
@@ -182,8 +197,11 @@ void FRHIDescriptorAllocator::Free(uint32 Offset, uint32 NumDescriptors)
 	}
 }
 
-FRHIHeapDescriptorAllocator::FRHIHeapDescriptorAllocator(ERHIDescriptorHeapType InType, uint32 InDescriptorCount)
-	: FRHIDescriptorAllocator(InDescriptorCount)
+////////////////////////////////////////////////////////////////////////////////////////////////////////////////////////
+// FRHIHeapDescriptorAllocator
+
+FRHIHeapDescriptorAllocator::FRHIHeapDescriptorAllocator(ERHIDescriptorHeapType InType, uint32 InDescriptorCount, TConstArrayView<TStatId> InStats)
+	: FRHIDescriptorAllocator(InDescriptorCount, InStats)
 	, Type(InType)
 {
 }
@@ -210,4 +228,28 @@ bool FRHIHeapDescriptorAllocator::Allocate(uint32 NumDescriptors, uint32& OutSlo
 void FRHIHeapDescriptorAllocator::Free(uint32 Slot, uint32 NumDescriptors)
 {
 	FRHIDescriptorAllocator::Free(Slot, NumDescriptors);
-}
\ No newline at end of file
+}
+
+////////////////////////////////////////////////////////////////////////////////////////////////////////////////////////
+// FRHIOffsetHeapDescriptorAllocator
+
+FRHIOffsetHeapDescriptorAllocator::FRHIOffsetHeapDescriptorAllocator(ERHIDescriptorHeapType InType, uint32 InDescriptorCount, uint32 InHeapOffset, TConstArrayView<TStatId> InStats)
+	: FRHIHeapDescriptorAllocator(InType, InDescriptorCount, InStats)
+	, HeapOffset(InHeapOffset)
+{
+}
+
+FRHIDescriptorHandle FRHIOffsetHeapDescriptorAllocator::Allocate()
+{
+	const FRHIDescriptorHandle AlocatorHandle = FRHIHeapDescriptorAllocator::Allocate();
+	return FRHIDescriptorHandle(AlocatorHandle.GetType(), AlocatorHandle.GetIndex() + HeapOffset);
+}
+
+void FRHIOffsetHeapDescriptorAllocator::Free(const FRHIDescriptorHandle InHandle)
+{
+	if (InHandle.IsValid())
+	{
+		const FRHIDescriptorHandle AdjustedHandle(InHandle.GetType(), InHandle.GetIndex() - HeapOffset);
+		FRHIHeapDescriptorAllocator::Free(AdjustedHandle);
+	}
+}
diff --git a/Engine/Source/Runtime/RHICore/Public/RHIDescriptorAllocator.h b/Engine/Source/Runtime/RHICore/Public/RHIDescriptorAllocator.h
--- a/Engine/Source/Runtime/RHICore/Public/RHIDescriptorAllocator.h
+++ b/Engine/Source/Runtime/RHICore/Public/RHIDescriptorAllocator.h
@@ -4,6 +4,7 @@
 #include "Containers/Array.h"
 #include "HAL/CriticalSection.h"
 #include "RHIDefinitions.h"
+#include "Stats/Stats.h"
 
 struct FRHIDescriptorAllocatorRange;
 
@@ -11,10 +12,10 @@ class RHICORE_API FRHIDescriptorAllocator
 {
 public:
 	FRHIDescriptorAllocator();
-	FRHIDescriptorAllocator(uint32 InNumDescriptors);
+	FRHIDescriptorAllocator(uint32 InNumDescriptors, TConstArrayView<TStatId> InStats);
 	~FRHIDescriptorAllocator();
 
-	void Init(uint32 InNumDescriptors);
+	void Init(uint32 InNumDescriptors, TConstArrayView<TStatId> InStats);
 	void Shutdown();
 
 	FRHIDescriptorHandle Allocate(ERHIDescriptorHeapType InType);
@@ -26,17 +27,41 @@ public:
 	inline uint32 GetCapacity() const { return Capacity; }
 
 private:
+	void RecordAlloc(uint32 Count)
+	{
+#if STATS
+		for (TStatId Stat : Stats)
+		{
+			INC_DWORD_STAT_BY_FName(Stat.GetName(), Count);
+		}
+#endif
+	}
+
+	void RecordFree(uint32 Count)
+	{
+#if STATS
+		for (TStatId Stat : Stats)
+		{
+			DEC_DWORD_STAT_BY_FName(Stat.GetName(), Count);
+		}
+#endif
+	}
+
 	TArray<FRHIDescriptorAllocatorRange> Ranges;
 	uint32 Capacity = 0;
 
 	FCriticalSection CriticalSection;
+
+#if STATS
+	TArray<TStatId> Stats;
+#endif
 };
 
 class RHICORE_API FRHIHeapDescriptorAllocator : protected FRHIDescriptorAllocator
 {
 public:
 	FRHIHeapDescriptorAllocator() = delete;
-	FRHIHeapDescriptorAllocator(ERHIDescriptorHeapType InType, uint32 InDescriptorCount);
+	FRHIHeapDescriptorAllocator(ERHIDescriptorHeapType InType, uint32 InDescriptorCount, TConstArrayView<TStatId> InStats);
 
 	FRHIDescriptorHandle Allocate();
 	void Free(FRHIDescriptorHandle InHandle);
@@ -52,3 +77,21 @@ public:
 private:
 	ERHIDescriptorHeapType Type;
 };
+
+class RHICORE_API FRHIOffsetHeapDescriptorAllocator : protected FRHIHeapDescriptorAllocator
+{
+public:
+	FRHIOffsetHeapDescriptorAllocator() = delete;
+	FRHIOffsetHeapDescriptorAllocator(ERHIDescriptorHeapType InType, uint32 InDescriptorCount, uint32 InHeapOffset, TConstArrayView<TStatId> InStats);
+
+	FRHIDescriptorHandle Allocate();
+	void Free(FRHIDescriptorHandle InHandle);
+
+	using FRHIHeapDescriptorAllocator::GetCapacity;
+	using FRHIHeapDescriptorAllocator::GetType;
+	using FRHIHeapDescriptorAllocator::HandlesAllocation;
+
+private:
+	// Offset from start of heap we belong to
+	uint32 HeapOffset;
+};
