Fix for non-deterministic acquiring of transient heaps, which causes resources to be placed multiple times on different heaps.

diff --git a/Engine/Source/Runtime/RHICore/Private/RHICoreTransientResourceAllocator.cpp b/Engine/Source/Runtime/RHICore/Private/RHICoreTransientResourceAllocator.cpp
--- a/Engine/Source/Runtime/RHICore/Private/RHICoreTransientResourceAllocator.cpp
+++ b/Engine/Source/Runtime/RHICore/Private/RHICoreTransientResourceAllocator.cpp
@@ -562,8 +562,6 @@ void FRHITransientHeap::Flush(uint64 AllocatorCycle, FRHITransientMemoryStats& O
 				const FRHITransientHeapAllocation& Allocation = Texture->GetHeapAllocation();
 				CommitSize = FMath::Max(CommitSize, Allocation.Offset + Allocation.Size);
 			}
-
-			TRACE_COUNTER_SET(TransientTextureCacheSize, Textures.GetSize());
 		}
 
 		{
@@ -574,8 +572,6 @@ void FRHITransientHeap::Flush(uint64 AllocatorCycle, FRHITransientMemoryStats& O
 				const FRHITransientHeapAllocation& Allocation = Buffer->GetHeapAllocation();
 				CommitSize = FMath::Max(CommitSize, Allocation.Offset + Allocation.Size);
 			}
-
-			TRACE_COUNTER_SET(TransientBufferCacheSize, Buffers.GetSize());
 		}
 
 		OverlapTracker.Reset();
@@ -610,7 +606,7 @@ FRHITransientHeap* FRHITransientHeapCache::Acquire(uint64 FirstAllocationSize, E
 {
 	FScopeLock Lock(&CriticalSection);
 
-	for (int32 HeapIndex = 0; HeapIndex < FreeList.Num(); ++HeapIndex)
+	for (int32 HeapIndex = FreeList.Num() - 1; HeapIndex >= 0; --HeapIndex)
 	{
 		FRHITransientHeap* Heap = FreeList[HeapIndex];
 
@@ -642,24 +638,13 @@ void FRHITransientHeapCache::Forfeit(TConstArrayView<FRHITransientHeap*> InForfe
 	FScopeLock Lock(&CriticalSection);
 
 	LiveList.Reserve(InForfeitedHeaps.Num());
-	for (FRHITransientHeap* Heap : InForfeitedHeaps)
+	for (int32 HeapIndex = InForfeitedHeaps.Num() - 1; HeapIndex >= 0; --HeapIndex)
 	{
+		FRHITransientHeap* Heap = InForfeitedHeaps[HeapIndex];
 		check(Heap->IsEmpty());
 		Heap->LastUsedGarbageCollectCycle = GarbageCollectCycle;
 		FreeList.Add(Heap);
 	}
-
-	Algo::Sort(FreeList, [](const FRHITransientHeap* LHS, const FRHITransientHeap* RHS)
-	{
-		// Sort by smaller heap first.
-		if (LHS->GetCapacity() != RHS->GetCapacity())
-		{
-			return LHS->GetCapacity() < RHS->GetCapacity();
-		}
-
-		// Sort next by most recently used first.
-		return LHS->GetLastUsedGarbageCollectCycle() >= RHS->GetLastUsedGarbageCollectCycle();
-	});
 }
 
 void FRHITransientHeapCache::GarbageCollect()
@@ -815,16 +800,43 @@ void FRHITransientResourceHeapAllocator::Flush(FRHICommandListImmediate& RHICmdL
 {
 	FRHITransientMemoryStats Stats;
 
+	uint32 NumBuffers = 0;
+	uint32 NumTextures = 0;
+
 	for (FRHITransientHeap* Heap : Heaps)
 	{
 		Heap->Flush(CurrentCycle, Stats, OutAllocationStats);
+
+		NumBuffers += Heap->Buffers.GetSize();
+		NumTextures += Heap->Textures.GetSize();
 	}
 
+	TRACE_COUNTER_SET(TransientBufferCacheSize, NumBuffers);
+	TRACE_COUNTER_SET(TransientTextureCacheSize, NumTextures);
+
 	if (DeallocationCount > 0)
 	{
-		int32 FirstForfeitIndex = Algo::Partition(Heaps.GetData(), Heaps.Num(), [](const FRHITransientHeap* Heap) { return !Heap->IsEmpty(); });
-		HeapCache.Forfeit(MakeArrayView(Heaps.GetData() + FirstForfeitIndex, Heaps.Num() - FirstForfeitIndex));
-		Heaps.SetNum(FirstForfeitIndex, false);
+		// This could be done more efficiently, but the number of heaps is small and the goal is to keep the list stable
+		// so that heaps are acquired in the same order each frame, because the resource caches are tied to heaps.
+		TArray<FRHITransientHeap*, FConcurrentLinearArrayAllocator> EmptyHeaps;
+		TArray<FRHITransientHeap*, FConcurrentLinearArrayAllocator> ActiveHeaps;
+		EmptyHeaps.Reserve(Heaps.Num());
+		ActiveHeaps.Reserve(Heaps.Num());
+
+		for (FRHITransientHeap* Heap : Heaps)
+		{
+			if (Heap->IsEmpty())
+			{
+				EmptyHeaps.Emplace(Heap);
+			}
+			else
+			{
+				ActiveHeaps.Emplace(Heap);
+			}
+		}
+
+		HeapCache.Forfeit(EmptyHeaps);
+		Heaps = ActiveHeaps;
 		DeallocationCount = 0;
 	}
 
