Moving a number of DXGI_FORMAT helpers to RHICore that were in both D3D11RHI and D3D12RHI.

diff --git a/Engine/Platforms/Hololens/Source/Runtime/RHICore/RHICore_Hololens.Build.cs b/Engine/Platforms/Hololens/Source/Runtime/RHICore/RHICore_Hololens.Build.cs
new file mode 100644
--- /dev/null
+++ b/Engine/Platforms/Hololens/Source/Runtime/RHICore/RHICore_Hololens.Build.cs
@@ -0,0 +1,11 @@
+// Copyright Epic Games, Inc. All Rights Reserved.
+
+using UnrealBuildTool;
+
+public class RHICore_HoloLens : RHICore
+{
+	public RHICore_HoloLens(ReadOnlyTargetRules Target) : base(Target)
+	{
+		PublicDefinitions.Add("RHICORE_PLATFORM_DXGI_H=<dxgi.h>");
+	}
+}
diff --git a/Engine/Source/Runtime/D3D12RHI/Private/D3D12Device.cpp b/Engine/Source/Runtime/D3D12RHI/Private/D3D12Device.cpp
--- a/Engine/Source/Runtime/D3D12RHI/Private/D3D12Device.cpp
+++ b/Engine/Source/Runtime/D3D12RHI/Private/D3D12Device.cpp
@@ -225,10 +225,10 @@ void FD3D12Device::SetupAfterDeviceCreation()
 		if (PlatformFormat != DXGI_FORMAT_UNKNOWN)
 		{
 			const D3D12_FEATURE_DATA_FORMAT_SUPPORT FormatSupport    = GetFormatSupport(Direct3DDevice, PlatformFormat);
-			const D3D12_FEATURE_DATA_FORMAT_SUPPORT SRVFormatSupport = GetFormatSupport(Direct3DDevice, FindShaderResourceDXGIFormat(PlatformFormat, false));
-			const D3D12_FEATURE_DATA_FORMAT_SUPPORT UAVFormatSupport = GetFormatSupport(Direct3DDevice, FindUnorderedAccessDXGIFormat(PlatformFormat));
-			const D3D12_FEATURE_DATA_FORMAT_SUPPORT RTVFormatSupport = GetFormatSupport(Direct3DDevice, FindShaderResourceDXGIFormat(PlatformFormat, false));
-			const D3D12_FEATURE_DATA_FORMAT_SUPPORT DSVFormatSupport = GetFormatSupport(Direct3DDevice, FindDepthStencilDXGIFormat(PlatformFormat));
+			const D3D12_FEATURE_DATA_FORMAT_SUPPORT SRVFormatSupport = GetFormatSupport(Direct3DDevice, UE::DXGIUtilities::FindShaderResourceFormat(PlatformFormat, false));
+			const D3D12_FEATURE_DATA_FORMAT_SUPPORT UAVFormatSupport = GetFormatSupport(Direct3DDevice, UE::DXGIUtilities::FindUnorderedAccessFormat(PlatformFormat));
+			const D3D12_FEATURE_DATA_FORMAT_SUPPORT RTVFormatSupport = GetFormatSupport(Direct3DDevice, UE::DXGIUtilities::FindShaderResourceFormat(PlatformFormat, false));
+			const D3D12_FEATURE_DATA_FORMAT_SUPPORT DSVFormatSupport = GetFormatSupport(Direct3DDevice, UE::DXGIUtilities::FindDepthStencilFormat(PlatformFormat));
 
 			auto ConvertCap1 = [&Capabilities](const D3D12_FEATURE_DATA_FORMAT_SUPPORT& InSupport, EPixelFormatCapabilities UnrealCap, D3D12_FORMAT_SUPPORT1 InFlags)
 			{
diff --git a/Engine/Source/Runtime/D3D12RHI/Private/D3D12Pipelinestate.cpp b/Engine/Source/Runtime/D3D12RHI/Private/D3D12Pipelinestate.cpp
--- a/Engine/Source/Runtime/D3D12RHI/Private/D3D12Pipelinestate.cpp
+++ b/Engine/Source/Runtime/D3D12RHI/Private/D3D12Pipelinestate.cpp
@@ -25,6 +25,30 @@ static FAutoConsoleVariableRef CVarPSOPrecacheKeepLowLevel(
 
 /// @cond DOXYGEN_WARNINGS
 
+static void TranslateRenderTargetFormats(
+	const FGraphicsPipelineStateInitializer& PsoInit,
+	D3D12_RT_FORMAT_ARRAY& RTFormatArray,
+	DXGI_FORMAT& DSVFormat)
+{
+	RTFormatArray.NumRenderTargets = PsoInit.ComputeNumValidRenderTargets();
+
+	for (uint32 RTIdx = 0; RTIdx < PsoInit.RenderTargetsEnabled; ++RTIdx)
+	{
+		checkSlow(PsoInit.RenderTargetFormats[RTIdx] == PF_Unknown || GPixelFormats[PsoInit.RenderTargetFormats[RTIdx]].Supported);
+
+		DXGI_FORMAT PlatformFormat = (DXGI_FORMAT)GPixelFormats[PsoInit.RenderTargetFormats[RTIdx]].PlatformFormat;
+		ETextureCreateFlags Flags = PsoInit.RenderTargetFlags[RTIdx];
+
+		RTFormatArray.RTFormats[RTIdx] = UE::DXGIUtilities::FindShaderResourceFormat(UE::DXGIUtilities::GetPlatformTextureResourceFormat(PlatformFormat, Flags), EnumHasAnyFlags(Flags, ETextureCreateFlags::SRGB));
+	}
+
+	checkSlow(PsoInit.DepthStencilTargetFormat == PF_Unknown || GPixelFormats[PsoInit.DepthStencilTargetFormat].Supported);
+
+	DXGI_FORMAT PlatformFormat = (DXGI_FORMAT)GPixelFormats[PsoInit.DepthStencilTargetFormat].PlatformFormat;
+
+	DSVFormat = UE::DXGIUtilities::FindDepthStencilFormat(UE::DXGIUtilities::GetPlatformTextureResourceFormat(PlatformFormat, PsoInit.DepthStencilTargetFlag));
+}
+
 static FD3D12LowLevelGraphicsPipelineStateDesc GetLowLevelGraphicsPipelineStateDesc(const FGraphicsPipelineStateInitializer& Initializer, const FD3D12RootSignature* RootSignature)
 {
 	FD3D12LowLevelGraphicsPipelineStateDesc Desc{};
diff --git a/Engine/Source/Runtime/D3D12RHI/Private/D3D12RHI.cpp b/Engine/Source/Runtime/D3D12RHI/Private/D3D12RHI.cpp
--- a/Engine/Source/Runtime/D3D12RHI/Private/D3D12RHI.cpp
+++ b/Engine/Source/Runtime/D3D12RHI/Private/D3D12RHI.cpp
@@ -620,8 +620,8 @@ ID3D12GraphicsCommandList* FD3D12DynamicRHI::RHIGetGraphicsCommandList(uint32 In
 
 DXGI_FORMAT FD3D12DynamicRHI::RHIGetSwapChainFormat(EPixelFormat InFormat) const
 {
-	const DXGI_FORMAT PlatformFormat = D3D12RHI::FindDepthStencilDXGIFormat(static_cast<DXGI_FORMAT>(GPixelFormats[InFormat].PlatformFormat));
-	return D3D12RHI::FindShaderResourceDXGIFormat(PlatformFormat, true);
+	const DXGI_FORMAT PlatformFormat = UE::DXGIUtilities::FindDepthStencilFormat(static_cast<DXGI_FORMAT>(GPixelFormats[InFormat].PlatformFormat));
+	return UE::DXGIUtilities::FindShaderResourceFormat(PlatformFormat, true);
 }
 
 ID3D12Resource* FD3D12DynamicRHI::RHIGetResource(FRHIBuffer* InBuffer) const
diff --git a/Engine/Source/Runtime/D3D12RHI/Private/D3D12RHIPrivate.h b/Engine/Source/Runtime/D3D12RHI/Private/D3D12RHIPrivate.h
--- a/Engine/Source/Runtime/D3D12RHI/Private/D3D12RHIPrivate.h
+++ b/Engine/Source/Runtime/D3D12RHI/Private/D3D12RHIPrivate.h
@@ -26,6 +26,7 @@
 // Dependencies.
 #include "CoreMinimal.h"
 #include "ID3D12DynamicRHI.h"
+#include "DXGIUtilities.h"
 #include "GPUProfiler.h"
 #include "ShaderCore.h"
 #include "HDRHelper.h"
@@ -809,193 +810,6 @@ public:
 	}
 };
 
-// This namespace is needed to avoid a name clash with D3D11 RHI when linked together in monolithic builds. Otherwise the linker will just pick any variant instead of each RHI using their own version.
-namespace D3D12RHI
-{
-
-inline DXGI_FORMAT FindSharedResourceDXGIFormat(DXGI_FORMAT InFormat, bool bSRGB)
-{
-	if (bSRGB)
-	{
-		switch (InFormat)
-		{
-		case DXGI_FORMAT_B8G8R8X8_TYPELESS: return DXGI_FORMAT_B8G8R8X8_UNORM_SRGB;
-		case DXGI_FORMAT_B8G8R8A8_TYPELESS: return DXGI_FORMAT_B8G8R8A8_UNORM_SRGB;
-		case DXGI_FORMAT_R8G8B8A8_TYPELESS: return DXGI_FORMAT_R8G8B8A8_UNORM_SRGB;
-		case DXGI_FORMAT_BC1_TYPELESS:      return DXGI_FORMAT_BC1_UNORM_SRGB;
-		case DXGI_FORMAT_BC2_TYPELESS:      return DXGI_FORMAT_BC2_UNORM_SRGB;
-		case DXGI_FORMAT_BC3_TYPELESS:      return DXGI_FORMAT_BC3_UNORM_SRGB;
-		case DXGI_FORMAT_BC7_TYPELESS:      return DXGI_FORMAT_BC7_UNORM_SRGB;
-		};
-	}
-	else
-	{
-		switch (InFormat)
-		{
-		case DXGI_FORMAT_B8G8R8X8_TYPELESS: return DXGI_FORMAT_B8G8R8X8_UNORM;
-		case DXGI_FORMAT_B8G8R8A8_TYPELESS: return DXGI_FORMAT_B8G8R8A8_UNORM;
-		case DXGI_FORMAT_R8G8B8A8_TYPELESS: return DXGI_FORMAT_R8G8B8A8_UNORM;
-		case DXGI_FORMAT_BC1_TYPELESS:      return DXGI_FORMAT_BC1_UNORM;
-		case DXGI_FORMAT_BC2_TYPELESS:      return DXGI_FORMAT_BC2_UNORM;
-		case DXGI_FORMAT_BC3_TYPELESS:      return DXGI_FORMAT_BC3_UNORM;
-		case DXGI_FORMAT_BC7_TYPELESS:      return DXGI_FORMAT_BC7_UNORM;
-		};
-	}
-
-	switch (InFormat)
-	{
-	case DXGI_FORMAT_R32G32B32A32_TYPELESS: return DXGI_FORMAT_R32G32B32A32_UINT;
-	case DXGI_FORMAT_R32G32B32_TYPELESS:    return DXGI_FORMAT_R32G32B32_UINT;
-	case DXGI_FORMAT_R16G16B16A16_TYPELESS: return DXGI_FORMAT_R16G16B16A16_UNORM;
-	case DXGI_FORMAT_R32G32_TYPELESS:       return DXGI_FORMAT_R32G32_UINT;
-	case DXGI_FORMAT_R10G10B10A2_TYPELESS:  return DXGI_FORMAT_R10G10B10A2_UNORM;
-	case DXGI_FORMAT_R16G16_TYPELESS:       return DXGI_FORMAT_R16G16_UNORM;
-	case DXGI_FORMAT_R8G8_TYPELESS:         return DXGI_FORMAT_R8G8_UNORM;
-	case DXGI_FORMAT_R8_TYPELESS:           return DXGI_FORMAT_R8_UNORM;
-
-	case DXGI_FORMAT_BC4_TYPELESS:         return DXGI_FORMAT_BC4_UNORM;
-	case DXGI_FORMAT_BC5_TYPELESS:         return DXGI_FORMAT_BC5_UNORM;
-	case DXGI_FORMAT_R24G8_TYPELESS:       return DXGI_FORMAT_R24_UNORM_X8_TYPELESS;
-	case DXGI_FORMAT_R32_TYPELESS:         return DXGI_FORMAT_R32_FLOAT;
-	case DXGI_FORMAT_R16_TYPELESS:         return DXGI_FORMAT_R16_UNORM;
-
-	case DXGI_FORMAT_R32G8X24_TYPELESS:    return DXGI_FORMAT_R32_FLOAT_X8X24_TYPELESS;
-	}
-
-	return InFormat;
-}
-
-inline DXGI_FORMAT FindDepthStencilResourceDXGIFormat(DXGI_FORMAT InFormat)
-{
-	switch (InFormat)
-	{
-	case DXGI_FORMAT_R32_FLOAT: return DXGI_FORMAT_R32_TYPELESS;
-	case DXGI_FORMAT_R16_FLOAT: return DXGI_FORMAT_R16_TYPELESS;
-	}
-
-	return InFormat;
-}
-
-inline DXGI_FORMAT GetPlatformTextureResourceFormat(DXGI_FORMAT InFormat, ETextureCreateFlags InFlags)
-{
-	// Find valid shared texture format
-	if (EnumHasAnyFlags(InFlags, TexCreate_Shared))
-	{
-		return FindSharedResourceDXGIFormat(InFormat, EnumHasAnyFlags(InFlags, TexCreate_SRGB));
-	}
-	if (EnumHasAnyFlags(InFlags, TexCreate_DepthStencilTargetable))
-	{
-		return FindDepthStencilResourceDXGIFormat(InFormat);
-	}
-
-	return InFormat;
-}
-
-/** Find an appropriate DXGI format for the input format and SRGB setting. */
-inline DXGI_FORMAT FindShaderResourceDXGIFormat(DXGI_FORMAT InFormat, bool bSRGB)
-{
-	if (bSRGB)
-	{
-		switch (InFormat)
-		{
-		case DXGI_FORMAT_B8G8R8A8_TYPELESS:    return DXGI_FORMAT_B8G8R8A8_UNORM_SRGB;
-		case DXGI_FORMAT_R8G8B8A8_TYPELESS:    return DXGI_FORMAT_R8G8B8A8_UNORM_SRGB;
-		case DXGI_FORMAT_BC1_TYPELESS:         return DXGI_FORMAT_BC1_UNORM_SRGB;
-		case DXGI_FORMAT_BC2_TYPELESS:         return DXGI_FORMAT_BC2_UNORM_SRGB;
-		case DXGI_FORMAT_BC3_TYPELESS:         return DXGI_FORMAT_BC3_UNORM_SRGB;
-		case DXGI_FORMAT_BC7_TYPELESS:         return DXGI_FORMAT_BC7_UNORM_SRGB;
-		};
-	}
-	else
-	{
-		switch (InFormat)
-		{
-		case DXGI_FORMAT_B8G8R8A8_TYPELESS: return DXGI_FORMAT_B8G8R8A8_UNORM;
-		case DXGI_FORMAT_R8G8B8A8_TYPELESS: return DXGI_FORMAT_R8G8B8A8_UNORM;
-		case DXGI_FORMAT_BC1_TYPELESS:      return DXGI_FORMAT_BC1_UNORM;
-		case DXGI_FORMAT_BC2_TYPELESS:      return DXGI_FORMAT_BC2_UNORM;
-		case DXGI_FORMAT_BC3_TYPELESS:      return DXGI_FORMAT_BC3_UNORM;
-		case DXGI_FORMAT_BC7_TYPELESS:      return DXGI_FORMAT_BC7_UNORM;
-		};
-	}
-	switch (InFormat)
-	{
-	case DXGI_FORMAT_R24G8_TYPELESS: return DXGI_FORMAT_R24_UNORM_X8_TYPELESS;
-	case DXGI_FORMAT_R32_TYPELESS: return DXGI_FORMAT_R32_FLOAT;
-	case DXGI_FORMAT_R16_TYPELESS: return DXGI_FORMAT_R16_UNORM;
-
-	case DXGI_FORMAT_R32G8X24_TYPELESS: return DXGI_FORMAT_R32_FLOAT_X8X24_TYPELESS;
-	}
-	return InFormat;
-}
-
-/** Find an appropriate DXGI format unordered access of the raw format. */
-inline DXGI_FORMAT FindUnorderedAccessDXGIFormat(DXGI_FORMAT InFormat)
-{
-	switch (InFormat)
-	{
-	case DXGI_FORMAT_B8G8R8A8_TYPELESS: return DXGI_FORMAT_B8G8R8A8_UNORM;
-	case DXGI_FORMAT_R8G8B8A8_TYPELESS: return DXGI_FORMAT_R8G8B8A8_UNORM;
-	case DXGI_FORMAT_R32G8X24_TYPELESS: return DXGI_FORMAT_R32_FLOAT_X8X24_TYPELESS;
-	}
-	return InFormat;
-}
-
-/** Find the appropriate depth-stencil targetable DXGI format for the given format. */
-inline DXGI_FORMAT FindDepthStencilDXGIFormat(DXGI_FORMAT InFormat)
-{
-	switch (InFormat)
-	{
-	case DXGI_FORMAT_R24G8_TYPELESS   : return DXGI_FORMAT_D24_UNORM_S8_UINT;
-	case DXGI_FORMAT_R32G8X24_TYPELESS: return DXGI_FORMAT_D32_FLOAT_S8X24_UINT;
-	case DXGI_FORMAT_R32_TYPELESS     : return DXGI_FORMAT_D32_FLOAT;
-	case DXGI_FORMAT_R16_TYPELESS     : return DXGI_FORMAT_D16_UNORM;
-	};
-	return InFormat;
-}
-
-/**
-* Returns whether the given format contains stencil information.
-* Must be passed a format returned by FindDepthStencilDXGIFormat, so that typeless versions are converted to their corresponding depth stencil view format.
-*/
-inline bool HasStencilBits(DXGI_FORMAT InFormat)
-{
-	switch (InFormat)
-	{
-	case DXGI_FORMAT_D24_UNORM_S8_UINT:
-	case DXGI_FORMAT_D32_FLOAT_S8X24_UINT:
-		return true;
-	};
-
-	return false;
-}
-
-static void TranslateRenderTargetFormats(
-	const FGraphicsPipelineStateInitializer &PsoInit,
-	D3D12_RT_FORMAT_ARRAY& RTFormatArray,
-	DXGI_FORMAT& DSVFormat)
-{
-	RTFormatArray.NumRenderTargets = PsoInit.ComputeNumValidRenderTargets();
-
-	for (uint32 RTIdx = 0; RTIdx < PsoInit.RenderTargetsEnabled; ++RTIdx)
-	{
-		checkSlow(PsoInit.RenderTargetFormats[RTIdx] == PF_Unknown || GPixelFormats[PsoInit.RenderTargetFormats[RTIdx]].Supported);
-
-		DXGI_FORMAT PlatformFormat = (DXGI_FORMAT)GPixelFormats[PsoInit.RenderTargetFormats[RTIdx]].PlatformFormat;
-		ETextureCreateFlags Flags = PsoInit.RenderTargetFlags[RTIdx];
-
-		RTFormatArray.RTFormats[RTIdx] = D3D12RHI::FindShaderResourceDXGIFormat( GetPlatformTextureResourceFormat(PlatformFormat, Flags), EnumHasAnyFlags(Flags, TexCreate_SRGB) );
-	}
-
-	checkSlow(PsoInit.DepthStencilTargetFormat == PF_Unknown || GPixelFormats[PsoInit.DepthStencilTargetFormat].Supported);
-
-	DXGI_FORMAT PlatformFormat = (DXGI_FORMAT)GPixelFormats[PsoInit.DepthStencilTargetFormat].PlatformFormat;
-
-	DSVFormat = D3D12RHI::FindDepthStencilDXGIFormat( GetPlatformTextureResourceFormat(PlatformFormat, PsoInit.DepthStencilTargetFlag) );
-}
-
-} // namespace D3D12RHI
-
 // Returns the given format as a string. Unsupported formats are treated as DXGI_FORMAT_UNKNOWN.
 const TCHAR* LexToString(DXGI_FORMAT Format);
 
diff --git a/Engine/Source/Runtime/D3D12RHI/Private/D3D12Resources.cpp b/Engine/Source/Runtime/D3D12RHI/Private/D3D12Resources.cpp
--- a/Engine/Source/Runtime/D3D12RHI/Private/D3D12Resources.cpp
+++ b/Engine/Source/Runtime/D3D12RHI/Private/D3D12Resources.cpp
@@ -49,8 +49,8 @@ TArray<DXGI_FORMAT, TInlineAllocator<4>> FD3D12ResourceDesc::GetCastableFormats(
 	TArray<DXGI_FORMAT, TInlineAllocator<4>> Result;
 
 	// We have to add the 'implied' castable formats for SRVs. Since we don't have any sRGB flags here, just add both formats.
-	Result.Emplace(FindShaderResourceDXGIFormat(Format, true));
-	Result.Emplace(FindShaderResourceDXGIFormat(Format, false));
+	Result.Emplace(UE::DXGIUtilities::FindShaderResourceFormat(Format, true));
+	Result.Emplace(UE::DXGIUtilities::FindShaderResourceFormat(Format, false));
 
 	if (UAVPixelFormat != PF_Unknown)
 	{
@@ -97,7 +97,7 @@ FD3D12Resource::FD3D12Resource(FD3D12Device* ParentDevice,
 	, Heap(InHeap)
 	, Desc(InDesc)
 	, HeapType(InHeapType)
-	, PlaneCount(::GetPlaneCount(Desc.Format))
+	, PlaneCount(UE::DXGIUtilities::GetPlaneCount(Desc.Format))
 	, bRequiresResourceStateTracking(true)
 	, bDepthStencil(false)
 	, bDeferDelete(true)
diff --git a/Engine/Source/Runtime/D3D12RHI/Private/D3D12SRV.cpp b/Engine/Source/Runtime/D3D12RHI/Private/D3D12SRV.cpp
--- a/Engine/Source/Runtime/D3D12RHI/Private/D3D12SRV.cpp
+++ b/Engine/Source/Runtime/D3D12RHI/Private/D3D12SRV.cpp
@@ -28,7 +28,7 @@ void FD3D12ShaderResourceView::CreateView(FResourceInfo const& InResource, D3D12
 	{
 		StrideInBytes = InD3DViewDesc.Format == DXGI_FORMAT_UNKNOWN
 			? InD3DViewDesc.Buffer.StructureByteStride
-			: D3D12GetFormatSizeInBytes(InD3DViewDesc.Format);
+			: UE::DXGIUtilities::GetFormatSizeInBytes(InD3DViewDesc.Format);
 
 		check(StrideInBytes > 0);
 
@@ -134,7 +134,7 @@ void FD3D12ShaderResourceView_RHI::CreateView()
 			else
 			{
 				SRVDesc.ViewDimension       = D3D12_SRV_DIMENSION_BUFFER;
-				SRVDesc.Format              = FindShaderResourceDXGIFormat(DXGI_FORMAT(GPixelFormats[Info.Format].PlatformFormat), false);
+				SRVDesc.Format              = UE::DXGIUtilities::FindShaderResourceFormat(DXGI_FORMAT(GPixelFormats[Info.Format].PlatformFormat), false);
 				SRVDesc.Buffer.FirstElement = (Info.OffsetInBytes + Buffer->ResourceLocation.GetOffsetFromBaseOfResource()) / Info.StrideInBytes;
 				SRVDesc.Buffer.NumElements  = Info.NumElements;
 
@@ -165,14 +165,14 @@ void FD3D12ShaderResourceView_RHI::CreateView()
 
 		auto const Info = ViewDesc.Texture.SRV.GetViewInfo(Texture);
 
-		DXGI_FORMAT const ViewFormat = FindShaderResourceDXGIFormat    (DXGI_FORMAT(GPixelFormats[Info.Format       ].PlatformFormat), Info.bSRGB       );
-		DXGI_FORMAT const BaseFormat = GetPlatformTextureResourceFormat(DXGI_FORMAT(GPixelFormats[TextureDesc.Format].PlatformFormat), TextureDesc.Flags);
+		DXGI_FORMAT const ViewFormat = UE::DXGIUtilities::FindShaderResourceFormat(DXGI_FORMAT(GPixelFormats[Info.Format       ].PlatformFormat), Info.bSRGB       );
+		DXGI_FORMAT const BaseFormat = UE::DXGIUtilities::GetPlatformTextureResourceFormat(DXGI_FORMAT(GPixelFormats[TextureDesc.Format].PlatformFormat), TextureDesc.Flags);
 
 		D3D12_SHADER_RESOURCE_VIEW_DESC SRVDesc{};
 		SRVDesc.Format = ViewFormat;
 		SRVDesc.Shader4ComponentMapping = D3D12_DEFAULT_SHADER_4_COMPONENT_MAPPING;
 
-		uint32 const PlaneSlice = GetPlaneSliceFromViewFormat(BaseFormat, ViewFormat);
+		uint32 const PlaneSlice = UE::DXGIUtilities::GetPlaneSliceFromViewFormat(BaseFormat, ViewFormat);
 		FRHIRange8 const PlaneRange(PlaneSlice, 1);
 
 		// No need to use Info.Dimension, since D3D supports mixing Texture2D view types.
diff --git a/Engine/Source/Runtime/D3D12RHI/Private/D3D12Texture.cpp b/Engine/Source/Runtime/D3D12RHI/Private/D3D12Texture.cpp
--- a/Engine/Source/Runtime/D3D12RHI/Private/D3D12Texture.cpp
+++ b/Engine/Source/Runtime/D3D12RHI/Private/D3D12Texture.cpp
@@ -541,7 +541,7 @@ FD3D12ResourceDesc FD3D12DynamicRHI::GetResourceDesc(const FRHITextureDesc& Text
 
 	check(TextureDesc.Extent.X > 0 && TextureDesc.Extent.Y > 0 && TextureDesc.NumMips > 0);
 
-	const DXGI_FORMAT PlatformResourceFormat = GetPlatformTextureResourceFormat((DXGI_FORMAT)GPixelFormats[TextureDesc.Format].PlatformFormat, TextureDesc.Flags);
+	const DXGI_FORMAT PlatformResourceFormat = UE::DXGIUtilities::GetPlatformTextureResourceFormat((DXGI_FORMAT)GPixelFormats[TextureDesc.Format].PlatformFormat, TextureDesc.Flags);
 
 	bool bDenyShaderResource = false;
 	if (TextureDesc.Dimension != ETextureDimension::Texture3D)
@@ -869,7 +869,7 @@ void SafeCreateTexture2D(FD3D12Device* pDevice,
 				TextureDesc.Width,
 				TextureDesc.Height,
 				TextureDesc.DepthOrArraySize,
-				GetD3D12TextureFormatString(TextureDesc.Format),
+				UE::DXGIUtilities::GetFormatString(TextureDesc.Format),
 				(uint32)TextureDesc.Format,
 				TextureDesc.MipLevels
 				);
@@ -899,7 +899,7 @@ FD3D12Texture* FD3D12DynamicRHI::CreateD3D12Texture(const FRHITextureCreateDesc&
 	// Get the resource desc
 	FD3D12ResourceDesc ResourceDesc = GetResourceDesc(CreateDesc);
 
-	const DXGI_FORMAT PlatformResourceFormat = GetPlatformTextureResourceFormat((DXGI_FORMAT)GPixelFormats[CreateDesc.Format].PlatformFormat, CreateDesc.Flags);
+	const DXGI_FORMAT PlatformResourceFormat = UE::DXGIUtilities::GetPlatformTextureResourceFormat((DXGI_FORMAT)GPixelFormats[CreateDesc.Format].PlatformFormat, CreateDesc.Flags);
 	bool bCreateRTV = EnumHasAnyFlags(ResourceDesc.Flags, D3D12_RESOURCE_FLAG_ALLOW_RENDER_TARGET);
 	bool bCreateDSV = EnumHasAnyFlags(ResourceDesc.Flags, D3D12_RESOURCE_FLAG_ALLOW_DEPTH_STENCIL);
 
@@ -908,14 +908,14 @@ FD3D12Texture* FD3D12DynamicRHI::CreateD3D12Texture(const FRHITextureCreateDesc&
 	D3D12_CLEAR_VALUE ClearValue;
 	if (bCreateDSV && CreateDesc.ClearValue.ColorBinding == EClearBinding::EDepthStencilBound)
 	{
-		const DXGI_FORMAT PlatformDepthStencilFormat = FindDepthStencilDXGIFormat(PlatformResourceFormat);
+		const DXGI_FORMAT PlatformDepthStencilFormat = UE::DXGIUtilities::FindDepthStencilFormat(PlatformResourceFormat);
 		ClearValue = CD3DX12_CLEAR_VALUE(PlatformDepthStencilFormat, CreateDesc.ClearValue.Value.DSValue.Depth, (uint8)CreateDesc.ClearValue.Value.DSValue.Stencil);
 		ClearValuePtr = &ClearValue;
 	}
 	else if (bCreateRTV && CreateDesc.ClearValue.ColorBinding == EClearBinding::EColorBound)
 	{
 		const bool bSRGB = EnumHasAnyFlags(CreateDesc.Flags, TexCreate_SRGB);
-		const DXGI_FORMAT PlatformRenderTargetFormat = FindShaderResourceDXGIFormat(PlatformResourceFormat, bSRGB);
+		const DXGI_FORMAT PlatformRenderTargetFormat = UE::DXGIUtilities::FindShaderResourceFormat(PlatformResourceFormat, bSRGB);
 		ClearValue = CD3DX12_CLEAR_VALUE(PlatformRenderTargetFormat, CreateDesc.ClearValue.Value.Color);
 		ClearValuePtr = &ClearValue;
 	}
@@ -1106,7 +1106,7 @@ FTextureRHIRef FD3D12DynamicRHI::RHIAsyncCreateTexture2D(uint32 SizeX, uint32 Si
 		.SetInitialState(ERHIAccess::SRVMask);
 
 	const DXGI_FORMAT PlatformResourceFormat = (DXGI_FORMAT)GPixelFormats[Format].PlatformFormat;
-	const DXGI_FORMAT PlatformShaderResourceFormat = FindShaderResourceDXGIFormat(PlatformResourceFormat, EnumHasAnyFlags(Flags, TexCreate_SRGB));
+	const DXGI_FORMAT PlatformShaderResourceFormat = UE::DXGIUtilities::FindShaderResourceFormat(PlatformResourceFormat, EnumHasAnyFlags(Flags, TexCreate_SRGB));
 	const D3D12_RESOURCE_DESC TextureDesc = CD3DX12_RESOURCE_DESC::Tex2D(
 		PlatformResourceFormat,
 		SizeX,
@@ -1176,7 +1176,7 @@ FTextureRHIRef FD3D12DynamicRHI::RHIAsyncCreateTexture2D(uint32 SizeX, uint32 Si
 		SRVDesc.ViewDimension = D3D12_SRV_DIMENSION_TEXTURE2D;
 		SRVDesc.Texture2D.MostDetailedMip = 0;
 		SRVDesc.Texture2D.MipLevels = NumMips;
-		SRVDesc.Texture2D.PlaneSlice = GetPlaneSliceFromViewFormat(PlatformResourceFormat, SRVDesc.Format);
+		SRVDesc.Texture2D.PlaneSlice = UE::DXGIUtilities::GetPlaneSliceFromViewFormat(PlatformResourceFormat, SRVDesc.Format);
 
 		// Create a wrapper for the SRV and set it on the texture
 		NewTexture->EmplaceSRV(SRVDesc);
@@ -1527,10 +1527,10 @@ void FD3D12Texture::CreateViews()
 	const FRHITextureDesc Desc = GetDesc();
 
 	const bool bSRGB = EnumHasAnyFlags(Desc.Flags, TexCreate_SRGB);
-	const DXGI_FORMAT PlatformResourceFormat = GetPlatformTextureResourceFormat((DXGI_FORMAT)GPixelFormats[Desc.Format].PlatformFormat, Desc.Flags);
-	const DXGI_FORMAT PlatformShaderResourceFormat = FindShaderResourceDXGIFormat(PlatformResourceFormat, bSRGB);
-	const DXGI_FORMAT PlatformRenderTargetFormat = FindShaderResourceDXGIFormat(PlatformResourceFormat, bSRGB);
-	const DXGI_FORMAT PlatformDepthStencilFormat = FindDepthStencilDXGIFormat(PlatformResourceFormat);
+	const DXGI_FORMAT PlatformResourceFormat = UE::DXGIUtilities::GetPlatformTextureResourceFormat((DXGI_FORMAT)GPixelFormats[Desc.Format].PlatformFormat, Desc.Flags);
+	const DXGI_FORMAT PlatformShaderResourceFormat = UE::DXGIUtilities::FindShaderResourceFormat(PlatformResourceFormat, bSRGB);
+	const DXGI_FORMAT PlatformRenderTargetFormat = UE::DXGIUtilities::FindShaderResourceFormat(PlatformResourceFormat, bSRGB);
+	const DXGI_FORMAT PlatformDepthStencilFormat = UE::DXGIUtilities::FindDepthStencilFormat(PlatformResourceFormat);
 
 	const bool bTexture2D = Desc.IsTexture2D();
 	const bool bTexture3D = Desc.IsTexture3D();
@@ -1603,7 +1603,7 @@ void FD3D12Texture::CreateViews()
 						RTVDesc.Texture2DArray.FirstArraySlice = SliceIndex;
 						RTVDesc.Texture2DArray.ArraySize = 1;
 						RTVDesc.Texture2DArray.MipSlice = MipIndex;
-						RTVDesc.Texture2DArray.PlaneSlice = GetPlaneSliceFromViewFormat(PlatformResourceFormat, RTVDesc.Format);
+						RTVDesc.Texture2DArray.PlaneSlice = UE::DXGIUtilities::GetPlaneSliceFromViewFormat(PlatformResourceFormat, RTVDesc.Format);
 
 						EmplaceRTV(RTVDesc, RTVIndex++);
 					}
@@ -1629,7 +1629,7 @@ void FD3D12Texture::CreateViews()
 							RTVDesc.Texture2DArray.FirstArraySlice = 0;
 							RTVDesc.Texture2DArray.ArraySize = ResourceDesc.DepthOrArraySize;
 							RTVDesc.Texture2DArray.MipSlice = MipIndex;
-							RTVDesc.Texture2DArray.PlaneSlice = GetPlaneSliceFromViewFormat(PlatformResourceFormat, RTVDesc.Format);
+							RTVDesc.Texture2DArray.PlaneSlice = UE::DXGIUtilities::GetPlaneSliceFromViewFormat(PlatformResourceFormat, RTVDesc.Format);
 						}
 					}
 					else
@@ -1643,7 +1643,7 @@ void FD3D12Texture::CreateViews()
 						{
 							RTVDesc.ViewDimension = D3D12_RTV_DIMENSION_TEXTURE2D;
 							RTVDesc.Texture2D.MipSlice = MipIndex;
-							RTVDesc.Texture2D.PlaneSlice = GetPlaneSliceFromViewFormat(PlatformResourceFormat, RTVDesc.Format);
+							RTVDesc.Texture2D.PlaneSlice = UE::DXGIUtilities::GetPlaneSliceFromViewFormat(PlatformResourceFormat, RTVDesc.Format);
 						}
 					}
 
@@ -1659,7 +1659,7 @@ void FD3D12Texture::CreateViews()
 
 		// Create a depth-stencil-view for the texture.
 		D3D12_DEPTH_STENCIL_VIEW_DESC DSVDesc = {};
-		DSVDesc.Format = FindDepthStencilDXGIFormat(PlatformResourceFormat);
+		DSVDesc.Format = UE::DXGIUtilities::FindDepthStencilFormat(PlatformResourceFormat);
 		if (bTextureArray || bCubeTexture)
 		{
 			if (bIsMultisampled)
@@ -1690,7 +1690,7 @@ void FD3D12Texture::CreateViews()
 			}
 		}
 
-		const bool HasStencil = HasStencilBits(DSVDesc.Format);
+		const bool HasStencil = UE::DXGIUtilities::HasStencilBits(DSVDesc.Format);
 		for (uint32 AccessType = 0; AccessType < FExclusiveDepthStencil::MaxIndex; ++AccessType)
 		{
 			// Create a read-only access views for the texture.
@@ -1740,7 +1740,7 @@ void FD3D12Texture::CreateViews()
 				SRVDesc.Texture2DArray.MipLevels = Desc.NumMips;
 				SRVDesc.Texture2DArray.FirstArraySlice = 0;
 				SRVDesc.Texture2DArray.ArraySize = ResourceDesc.DepthOrArraySize;
-				SRVDesc.Texture2DArray.PlaneSlice = GetPlaneSliceFromViewFormat(PlatformResourceFormat, SRVDesc.Format);
+				SRVDesc.Texture2DArray.PlaneSlice = UE::DXGIUtilities::GetPlaneSliceFromViewFormat(PlatformResourceFormat, SRVDesc.Format);
 			}
 		}
 		else if (bTexture3D)
@@ -1761,7 +1761,7 @@ void FD3D12Texture::CreateViews()
 				SRVDesc.ViewDimension = D3D12_SRV_DIMENSION_TEXTURE2D;
 				SRVDesc.Texture2D.MostDetailedMip = 0;
 				SRVDesc.Texture2D.MipLevels = Desc.NumMips;
-				SRVDesc.Texture2D.PlaneSlice = GetPlaneSliceFromViewFormat(PlatformResourceFormat, SRVDesc.Format);
+				SRVDesc.Texture2D.PlaneSlice = UE::DXGIUtilities::GetPlaneSliceFromViewFormat(PlatformResourceFormat, SRVDesc.Format);
 			}
 		}
 
@@ -2905,8 +2905,8 @@ FD3D12Texture* FD3D12DynamicRHI::CreateAliasedD3D12Texture2D(FD3D12Texture* Sour
 	const bool bSRGB = EnumHasAnyFlags(SourceTexture->GetFlags(), TexCreate_SRGB);
 
 	const DXGI_FORMAT PlatformResourceFormat = TextureDesc.Format;
-	const DXGI_FORMAT PlatformShaderResourceFormat = FindShaderResourceDXGIFormat(PlatformResourceFormat, bSRGB);
-	const DXGI_FORMAT PlatformRenderTargetFormat = FindShaderResourceDXGIFormat(PlatformResourceFormat, bSRGB);
+	const DXGI_FORMAT PlatformShaderResourceFormat = UE::DXGIUtilities::FindShaderResourceFormat(PlatformResourceFormat, bSRGB);
+	const DXGI_FORMAT PlatformRenderTargetFormat = UE::DXGIUtilities::FindShaderResourceFormat(PlatformResourceFormat, bSRGB);
 
 	const FString Name = SourceTexture->GetName().ToString() + TEXT("Alias");
 	FRHITextureCreateDesc CreateDesc(SourceTexture->GetDesc(), ERHIAccess::SRVMask, *Name);
diff --git a/Engine/Source/Runtime/D3D12RHI/Private/D3D12UAV.cpp b/Engine/Source/Runtime/D3D12RHI/Private/D3D12UAV.cpp
--- a/Engine/Source/Runtime/D3D12RHI/Private/D3D12UAV.cpp
+++ b/Engine/Source/Runtime/D3D12RHI/Private/D3D12UAV.cpp
@@ -28,7 +28,7 @@ void FD3D12UnorderedAccessView::CreateView(FResourceInfo const& InResource, D3D1
 	{
 		StrideInBytes = InD3DViewDesc.Format == DXGI_FORMAT_UNKNOWN
 			? InD3DViewDesc.Buffer.StructureByteStride
-			: D3D12GetFormatSizeInBytes(InD3DViewDesc.Format);
+			: UE::DXGIUtilities::GetFormatSizeInBytes(InD3DViewDesc.Format);
 
 		check(StrideInBytes > 0);
 
@@ -105,7 +105,7 @@ void FD3D12UnorderedAccessView_RHI::CreateView()
 				Flags = EFlags::NeedsCounter;
 
 			UAVDesc.ViewDimension       = D3D12_UAV_DIMENSION_BUFFER;
-			UAVDesc.Format              = FindUnorderedAccessDXGIFormat(DXGI_FORMAT(GPixelFormats[Info.Format].PlatformFormat));
+			UAVDesc.Format              = UE::DXGIUtilities::FindUnorderedAccessFormat(DXGI_FORMAT(GPixelFormats[Info.Format].PlatformFormat));
 			UAVDesc.Buffer.FirstElement = (Info.OffsetInBytes + Buffer->ResourceLocation.GetOffsetFromBaseOfResource()) / Info.StrideInBytes;
 			UAVDesc.Buffer.NumElements  = Info.NumElements;
 
@@ -141,13 +141,13 @@ void FD3D12UnorderedAccessView_RHI::CreateView()
 
 		auto const Info = ViewDesc.Texture.UAV.GetViewInfo(Texture);
 
-		DXGI_FORMAT const ViewFormat = FindUnorderedAccessDXGIFormat   (DXGI_FORMAT(GPixelFormats[Info.Format       ].PlatformFormat));
-		DXGI_FORMAT const BaseFormat = GetPlatformTextureResourceFormat(DXGI_FORMAT(GPixelFormats[TextureDesc.Format].PlatformFormat), TextureDesc.Flags);
+		DXGI_FORMAT const ViewFormat = UE::DXGIUtilities::FindUnorderedAccessFormat(DXGI_FORMAT(GPixelFormats[Info.Format       ].PlatformFormat));
+		DXGI_FORMAT const BaseFormat = UE::DXGIUtilities::GetPlatformTextureResourceFormat(DXGI_FORMAT(GPixelFormats[TextureDesc.Format].PlatformFormat), TextureDesc.Flags);
 
 		D3D12_UNORDERED_ACCESS_VIEW_DESC UAVDesc{};
 		UAVDesc.Format = ViewFormat;
 
-		uint32 const PlaneSlice = GetPlaneSliceFromViewFormat(BaseFormat, ViewFormat);
+		uint32 const PlaneSlice = UE::DXGIUtilities::GetPlaneSliceFromViewFormat(BaseFormat, ViewFormat);
 		FRHIRange8 const PlaneRange(PlaneSlice, 1);
 
 		// No need to use Info.Dimension, since D3D supports mixing Texture2D view types.
diff --git a/Engine/Source/Runtime/D3D12RHI/Private/D3D12Util.cpp b/Engine/Source/Runtime/D3D12RHI/Private/D3D12Util.cpp
--- a/Engine/Source/Runtime/D3D12RHI/Private/D3D12Util.cpp
+++ b/Engine/Source/Runtime/D3D12RHI/Private/D3D12Util.cpp
@@ -136,54 +136,6 @@ static FString GetD3D12ErrorString(HRESULT ErrorCode, ID3D12Device* Device)
 
 #undef D3DERR
 
-namespace D3D12RHI
-{
-	const TCHAR* GetD3D12TextureFormatString(DXGI_FORMAT TextureFormat)
-	{
-		static const TCHAR* EmptyString = TEXT("");
-		const TCHAR* TextureFormatText = EmptyString;
-#define D3DFORMATCASE(x) case x: TextureFormatText = TEXT(#x); break;
-		switch (TextureFormat)
-		{
-			D3DFORMATCASE(DXGI_FORMAT_R8G8B8A8_UNORM)
-				D3DFORMATCASE(DXGI_FORMAT_B8G8R8A8_UNORM)
-				D3DFORMATCASE(DXGI_FORMAT_B8G8R8X8_UNORM)
-				D3DFORMATCASE(DXGI_FORMAT_BC1_UNORM)
-				D3DFORMATCASE(DXGI_FORMAT_BC2_UNORM)
-				D3DFORMATCASE(DXGI_FORMAT_BC3_UNORM)
-				D3DFORMATCASE(DXGI_FORMAT_BC4_UNORM)
-				D3DFORMATCASE(DXGI_FORMAT_R16G16B16A16_FLOAT)
-				D3DFORMATCASE(DXGI_FORMAT_R32G32B32A32_FLOAT)
-				D3DFORMATCASE(DXGI_FORMAT_UNKNOWN)
-				D3DFORMATCASE(DXGI_FORMAT_R8_UNORM)
-				D3DFORMATCASE(DXGI_FORMAT_D32_FLOAT_S8X24_UINT)
-				D3DFORMATCASE(DXGI_FORMAT_R32_FLOAT_X8X24_TYPELESS)
-				D3DFORMATCASE(DXGI_FORMAT_R32G8X24_TYPELESS)
-				D3DFORMATCASE(DXGI_FORMAT_D24_UNORM_S8_UINT)
-				D3DFORMATCASE(DXGI_FORMAT_R24_UNORM_X8_TYPELESS)
-				D3DFORMATCASE(DXGI_FORMAT_R32_FLOAT)
-				D3DFORMATCASE(DXGI_FORMAT_R16G16_UINT)
-				D3DFORMATCASE(DXGI_FORMAT_R16G16_UNORM)
-				D3DFORMATCASE(DXGI_FORMAT_R16G16_SNORM)
-				D3DFORMATCASE(DXGI_FORMAT_R16G16_FLOAT)
-				D3DFORMATCASE(DXGI_FORMAT_R32G32_FLOAT)
-				D3DFORMATCASE(DXGI_FORMAT_R10G10B10A2_UNORM)
-				D3DFORMATCASE(DXGI_FORMAT_R16G16B16A16_UINT)
-				D3DFORMATCASE(DXGI_FORMAT_R8G8_SNORM)
-				D3DFORMATCASE(DXGI_FORMAT_BC5_UNORM)
-				D3DFORMATCASE(DXGI_FORMAT_R1_UNORM)
-				D3DFORMATCASE(DXGI_FORMAT_R8G8B8A8_TYPELESS)
-				D3DFORMATCASE(DXGI_FORMAT_B8G8R8A8_TYPELESS)
-				D3DFORMATCASE(DXGI_FORMAT_BC7_UNORM)
-				D3DFORMATCASE(DXGI_FORMAT_BC6H_UF16)
-		default: TextureFormatText = EmptyString;
-		}
-#undef D3DFORMATCASE
-		return TextureFormatText;
-	}
-}
-using namespace D3D12RHI;
-
 static FString GetD3D12TextureFlagString(uint32 TextureFlags)
 {
 	FString TextureFormatText = TEXT("");
@@ -947,7 +899,7 @@ namespace D3D12RHI
 		GD3DCallFailedCS.Lock();
 
 		const FString ErrorString = GetD3D12ErrorString(D3DResult, nullptr);
-		const TCHAR* D3DFormatString = GetD3D12TextureFormatString(TextureDesc.Format);
+		const TCHAR* D3DFormatString = UE::DXGIUtilities::GetFormatString(TextureDesc.Format);
 
 		UE_LOG(LogD3D12RHI, Error,
 			TEXT("%s failed \n at %s:%u \n with error %s, \n Size=%ix%ix%i Format=%s(0x%08X), NumMips=%i, Flags=%s"),
diff --git a/Engine/Source/Runtime/D3D12RHI/Private/D3D12Util.h b/Engine/Source/Runtime/D3D12RHI/Private/D3D12Util.h
--- a/Engine/Source/Runtime/D3D12RHI/Private/D3D12Util.h
+++ b/Engine/Source/Runtime/D3D12RHI/Private/D3D12Util.h
@@ -49,9 +49,6 @@ namespace D3D12RHI
 	extern void VerifyComRefCount(IUnknown* Object, int32 ExpectedRefs, const TCHAR* Code, const TCHAR* Filename, int32 Line);
 #define checkComRefCount(Obj,ExpectedRefs) VerifyComRefCount(Obj,ExpectedRefs,TEXT(#Obj),TEXT(__FILE__),__LINE__)
 
-	/** Returns a string for the provided DXGI format. */
-	const TCHAR* GetD3D12TextureFormatString(DXGI_FORMAT TextureFormat);
-
 	/** Checks if given GPU virtual address corresponds to any known resource allocations and logs results */
 	void LogPageFaultData(class FD3D12Adapter* InAdapter, FD3D12Device* InDevice, D3D12_GPU_VIRTUAL_ADDRESS InPageFaultAddress);
 	
@@ -456,13 +453,6 @@ inline D3D12_RESOURCE_STATES DetermineInitialResourceState(D3D12_HEAP_TYPE HeapT
 	}
 }
 
-static bool IsBlockCompressFormat(DXGI_FORMAT Format)
-{
-	// Returns true if BC1, BC2, BC3, BC4, BC5, BC6, BC7
-	return (Format >= DXGI_FORMAT_BC1_TYPELESS && Format <= DXGI_FORMAT_BC5_SNORM) ||
-		(Format >= DXGI_FORMAT_BC6H_TYPELESS && Format <= DXGI_FORMAT_BC7_UNORM_SRGB);
-}
-
 static inline uint64 GetTilesNeeded(uint32 Width, uint32 Height, uint32 Depth, const D3D12_TILE_SHAPE& Shape)
 {
 	return uint64((Width + Shape.WidthInTexels - 1) / Shape.WidthInTexels) *
@@ -470,79 +460,6 @@ static inline uint64 GetTilesNeeded(uint32 Width, uint32 Height, uint32 Depth, c
 		((Depth + Shape.DepthInTexels - 1) / Shape.DepthInTexels);
 }
 
-static uint32 GetWidthAlignment(DXGI_FORMAT Format)
-{
-	switch (Format)
-	{
-	case DXGI_FORMAT_R8G8_B8G8_UNORM: return 2;
-	case DXGI_FORMAT_G8R8_G8B8_UNORM: return 2;
-	case DXGI_FORMAT_NV12: return 2;
-	case DXGI_FORMAT_P010: return 2;
-	case DXGI_FORMAT_P016: return 2;
-	case DXGI_FORMAT_420_OPAQUE: return 2;
-	case DXGI_FORMAT_YUY2: return 2;
-	case DXGI_FORMAT_Y210: return 2;
-	case DXGI_FORMAT_Y216: return 2;
-	case DXGI_FORMAT_BC1_TYPELESS: return 4;
-	case DXGI_FORMAT_BC1_UNORM: return 4;
-	case DXGI_FORMAT_BC1_UNORM_SRGB: return 4;
-	case DXGI_FORMAT_BC2_TYPELESS: return 4;
-	case DXGI_FORMAT_BC2_UNORM: return 4;
-	case DXGI_FORMAT_BC2_UNORM_SRGB: return 4;
-	case DXGI_FORMAT_BC3_TYPELESS: return 4;
-	case DXGI_FORMAT_BC3_UNORM: return 4;
-	case DXGI_FORMAT_BC3_UNORM_SRGB: return 4;
-	case DXGI_FORMAT_BC4_TYPELESS: return 4;
-	case DXGI_FORMAT_BC4_UNORM: return 4;
-	case DXGI_FORMAT_BC4_SNORM: return 4;
-	case DXGI_FORMAT_BC5_TYPELESS: return 4;
-	case DXGI_FORMAT_BC5_UNORM: return 4;
-	case DXGI_FORMAT_BC5_SNORM: return 4;
-	case DXGI_FORMAT_BC6H_TYPELESS: return 4;
-	case DXGI_FORMAT_BC6H_UF16: return 4;
-	case DXGI_FORMAT_BC6H_SF16: return 4;
-	case DXGI_FORMAT_BC7_TYPELESS: return 4;
-	case DXGI_FORMAT_BC7_UNORM: return 4;
-	case DXGI_FORMAT_BC7_UNORM_SRGB: return 4;
-	case DXGI_FORMAT_NV11: return 4;
-	case DXGI_FORMAT_R1_UNORM: return 8;
-	default: return 1;
-	}
-}
-
-static uint32 GetHeightAlignment(DXGI_FORMAT Format)
-{
-	switch (Format)
-	{
-	case DXGI_FORMAT_NV12: return 2;
-	case DXGI_FORMAT_P010: return 2;
-	case DXGI_FORMAT_P016: return 2;
-	case DXGI_FORMAT_420_OPAQUE: return 2;
-	case DXGI_FORMAT_BC1_TYPELESS: return 4;
-	case DXGI_FORMAT_BC1_UNORM: return 4;
-	case DXGI_FORMAT_BC1_UNORM_SRGB: return 4;
-	case DXGI_FORMAT_BC2_TYPELESS: return 4;
-	case DXGI_FORMAT_BC2_UNORM: return 4;
-	case DXGI_FORMAT_BC2_UNORM_SRGB: return 4;
-	case DXGI_FORMAT_BC3_TYPELESS: return 4;
-	case DXGI_FORMAT_BC3_UNORM: return 4;
-	case DXGI_FORMAT_BC3_UNORM_SRGB: return 4;
-	case DXGI_FORMAT_BC4_TYPELESS: return 4;
-	case DXGI_FORMAT_BC4_UNORM: return 4;
-	case DXGI_FORMAT_BC4_SNORM: return 4;
-	case DXGI_FORMAT_BC5_TYPELESS: return 4;
-	case DXGI_FORMAT_BC5_UNORM: return 4;
-	case DXGI_FORMAT_BC5_SNORM: return 4;
-	case DXGI_FORMAT_BC6H_TYPELESS: return 4;
-	case DXGI_FORMAT_BC6H_UF16: return 4;
-	case DXGI_FORMAT_BC6H_SF16: return 4;
-	case DXGI_FORMAT_BC7_TYPELESS: return 4;
-	case DXGI_FORMAT_BC7_UNORM: return 4;
-	case DXGI_FORMAT_BC7_UNORM_SRGB: return 4;
-	default: return 1;
-	}
-}
-
 static void Get4KTileShape(D3D12_TILE_SHAPE* pTileShape, DXGI_FORMAT DXGIFormat, EPixelFormat UEFormat, D3D12_RESOURCE_DIMENSION Dimension, uint32 SampleCount)
 {
 	//Bits per unit
@@ -553,7 +470,7 @@ static void Get4KTileShape(D3D12_TILE_SHAPE* pTileShape, DXGI_FORMAT DXGIFormat,
 	case D3D12_RESOURCE_DIMENSION_BUFFER:
 	case D3D12_RESOURCE_DIMENSION_TEXTURE1D:
 	{
-		check(!IsBlockCompressFormat(DXGIFormat));
+		check(!UE::DXGIUtilities::IsBlockCompressedFormat(DXGIFormat));
 		pTileShape->WidthInTexels = (BPU == 0) ? 4096 : 4096 * 8 / BPU;
 		pTileShape->HeightInTexels = 1;
 		pTileShape->DepthInTexels = 1;
@@ -562,13 +479,13 @@ static void Get4KTileShape(D3D12_TILE_SHAPE* pTileShape, DXGI_FORMAT DXGIFormat,
 	case D3D12_RESOURCE_DIMENSION_TEXTURE2D:
 	{
 		pTileShape->DepthInTexels = 1;
-		if (IsBlockCompressFormat(DXGIFormat))
+		if (UE::DXGIUtilities::IsBlockCompressedFormat(DXGIFormat))
 		{
 			// Currently only supported block sizes are 64 and 128.
 			// These equations calculate the size in texels for a tile. It relies on the fact that 16*16*16 blocks fit in a tile if the block size is 128 bits.
 			check(BPU == 64 || BPU == 128);
-			pTileShape->WidthInTexels = 16 * GetWidthAlignment(DXGIFormat);
-			pTileShape->HeightInTexels = 16 * GetHeightAlignment(DXGIFormat);
+			pTileShape->WidthInTexels = 16 * UE::DXGIUtilities::GetWidthAlignment(DXGIFormat);
+			pTileShape->HeightInTexels = 16 * UE::DXGIUtilities::GetHeightAlignment(DXGIFormat);
 			if (BPU == 64)
 			{
 				// If bits per block are 64 we double width so it takes up the full tile size.
@@ -638,21 +555,21 @@ static void Get4KTileShape(D3D12_TILE_SHAPE* pTileShape, DXGI_FORMAT DXGIFormat,
 				check(false);
 			}
 
-			check(GetWidthAlignment(DXGIFormat) == 1);
-			check(GetHeightAlignment(DXGIFormat) == 1);
+			check(UE::DXGIUtilities::GetWidthAlignment(DXGIFormat) == 1);
+			check(UE::DXGIUtilities::GetHeightAlignment(DXGIFormat) == 1);
 		}
 
 		break;
 	}
 	case D3D12_RESOURCE_DIMENSION_TEXTURE3D:
 	{
-		if (IsBlockCompressFormat(DXGIFormat))
+		if (UE::DXGIUtilities::IsBlockCompressedFormat(DXGIFormat))
 		{
 			// Currently only supported block sizes are 64 and 128.
 			// These equations calculate the size in texels for a tile. It relies on the fact that 16*16*16 blocks fit in a tile if the block size is 128 bits.
 			check(BPU == 64 || BPU == 128);
-			pTileShape->WidthInTexels = 8 * GetWidthAlignment(DXGIFormat);
-			pTileShape->HeightInTexels = 8 * GetHeightAlignment(DXGIFormat);
+			pTileShape->WidthInTexels = 8 * UE::DXGIUtilities::GetWidthAlignment(DXGIFormat);
+			pTileShape->HeightInTexels = 8 * UE::DXGIUtilities::GetHeightAlignment(DXGIFormat);
 			pTileShape->DepthInTexels = 4;
 			if (BPU == 64)
 			{
@@ -700,8 +617,8 @@ static void Get4KTileShape(D3D12_TILE_SHAPE* pTileShape, DXGI_FORMAT DXGIFormat,
 				check(false);
 			}
 
-			check(GetWidthAlignment(DXGIFormat) == 1);
-			check(GetHeightAlignment(DXGIFormat) == 1);
+			check(UE::DXGIUtilities::GetWidthAlignment(DXGIFormat) == 1);
+			check(UE::DXGIUtilities::GetHeightAlignment(DXGIFormat) == 1);
 		}
 	}
 	break;
@@ -802,209 +719,6 @@ FORCEINLINE_DEBUGGABLE uint32 GetMaxMSAAQuality(uint32 SampleCount)
 	return 0xffffffff;
 }
 
-/** Find the appropriate depth-stencil typeless DXGI format for the given format. */
-inline DXGI_FORMAT FindDepthStencilParentDXGIFormat(DXGI_FORMAT InFormat)
-{
-	switch (InFormat)
-	{
-	case DXGI_FORMAT_D24_UNORM_S8_UINT      : return DXGI_FORMAT_R24G8_TYPELESS;
-	case DXGI_FORMAT_X24_TYPELESS_G8_UINT   : return DXGI_FORMAT_R24G8_TYPELESS;
-
-	case DXGI_FORMAT_D32_FLOAT_S8X24_UINT   : return DXGI_FORMAT_R32G8X24_TYPELESS;
-	case DXGI_FORMAT_X32_TYPELESS_G8X24_UINT: return DXGI_FORMAT_R32G8X24_TYPELESS;
-
-	case DXGI_FORMAT_D32_FLOAT              : return DXGI_FORMAT_R32_TYPELESS;
-	case DXGI_FORMAT_D16_UNORM              : return DXGI_FORMAT_R16_TYPELESS;
-	};
-	return InFormat;
-}
-
-static uint8 GetPlaneSliceFromViewFormat(DXGI_FORMAT ResourceFormat, DXGI_FORMAT ViewFormat)
-{
-	// Currently, the only planar resources used are depth-stencil formats
-	switch (FindDepthStencilParentDXGIFormat(ResourceFormat))
-	{
-	case DXGI_FORMAT_R24G8_TYPELESS:
-		switch (ViewFormat)
-		{
-		case DXGI_FORMAT_R24_UNORM_X8_TYPELESS: return 0;
-		case DXGI_FORMAT_X24_TYPELESS_G8_UINT : return 1;
-		}
-		break;
-	
-	case DXGI_FORMAT_R32G8X24_TYPELESS:
-		switch (ViewFormat)
-		{
-		case DXGI_FORMAT_R32_FLOAT_X8X24_TYPELESS: return 0;
-		case DXGI_FORMAT_X32_TYPELESS_G8X24_UINT : return 1;
-		}
-		break;
-
-	case DXGI_FORMAT_NV12:
-		switch (ViewFormat)
-		{
-		case DXGI_FORMAT_R8_UNORM  : return 0;
-		case DXGI_FORMAT_R8G8_UNORM: return 1;
-		}
-		break;
-	}
-
-	return 0;
-}
-
-static uint8 GetPlaneCount(DXGI_FORMAT Format)
-{
-	// Currently, the only planar resources used are depth-stencil formats
-	// Note there is a D3D12 helper for this, D3D12GetFormatPlaneCount
-	switch (FindDepthStencilParentDXGIFormat(Format))
-	{
-	case DXGI_FORMAT_R24G8_TYPELESS:
-	case DXGI_FORMAT_R32G8X24_TYPELESS:
-		return 2;
-	default:
-		return 1;
-	}
-}
-
-static uint32 D3D12GetFormatSizeInBits(DXGI_FORMAT Format)
-{
-	switch (Format)
-	{
-	default: checkNoEntry(); [[fallthrough]];
-	case DXGI_FORMAT_UNKNOWN:
-		return 0;
-
-	case DXGI_FORMAT_R32G32B32A32_TYPELESS:
-	case DXGI_FORMAT_R32G32B32A32_FLOAT:
-	case DXGI_FORMAT_R32G32B32A32_UINT:
-	case DXGI_FORMAT_R32G32B32A32_SINT:
-	case DXGI_FORMAT_BC2_TYPELESS:
-	case DXGI_FORMAT_BC2_UNORM:
-	case DXGI_FORMAT_BC2_UNORM_SRGB:
-	case DXGI_FORMAT_BC3_TYPELESS:
-	case DXGI_FORMAT_BC3_UNORM:
-	case DXGI_FORMAT_BC3_UNORM_SRGB:
-	case DXGI_FORMAT_BC5_TYPELESS:
-	case DXGI_FORMAT_BC5_UNORM:
-	case DXGI_FORMAT_BC5_SNORM:
-	case DXGI_FORMAT_BC6H_TYPELESS:
-	case DXGI_FORMAT_BC6H_UF16:
-	case DXGI_FORMAT_BC6H_SF16:
-	case DXGI_FORMAT_BC7_TYPELESS:
-	case DXGI_FORMAT_BC7_UNORM:
-	case DXGI_FORMAT_BC7_UNORM_SRGB:
-		return 128;
-
-	case DXGI_FORMAT_R32G32B32_TYPELESS:
-	case DXGI_FORMAT_R32G32B32_FLOAT:
-	case DXGI_FORMAT_R32G32B32_UINT:
-	case DXGI_FORMAT_R32G32B32_SINT:
-		return 96;
-
-	case DXGI_FORMAT_R16G16B16A16_TYPELESS:
-	case DXGI_FORMAT_R16G16B16A16_FLOAT:
-	case DXGI_FORMAT_R16G16B16A16_UNORM:
-	case DXGI_FORMAT_R16G16B16A16_UINT:
-	case DXGI_FORMAT_R16G16B16A16_SNORM:
-	case DXGI_FORMAT_R16G16B16A16_SINT:
-	case DXGI_FORMAT_R32G32_TYPELESS:
-	case DXGI_FORMAT_R32G32_FLOAT:
-	case DXGI_FORMAT_R32G32_UINT:
-	case DXGI_FORMAT_R32G32_SINT:
-	case DXGI_FORMAT_R32G8X24_TYPELESS:
-	case DXGI_FORMAT_D32_FLOAT_S8X24_UINT:
-	case DXGI_FORMAT_R32_FLOAT_X8X24_TYPELESS:
-	case DXGI_FORMAT_X32_TYPELESS_G8X24_UINT:
-	case DXGI_FORMAT_BC1_TYPELESS:
-	case DXGI_FORMAT_BC1_UNORM:
-	case DXGI_FORMAT_BC1_UNORM_SRGB:
-	case DXGI_FORMAT_BC4_TYPELESS:
-	case DXGI_FORMAT_BC4_UNORM:
-	case DXGI_FORMAT_BC4_SNORM:
-	case DXGI_FORMAT_Y416:
-	case DXGI_FORMAT_Y210:
-	case DXGI_FORMAT_Y216:
-		return 64;
-
-	case DXGI_FORMAT_R10G10B10A2_TYPELESS:
-	case DXGI_FORMAT_R10G10B10A2_UNORM:
-	case DXGI_FORMAT_R10G10B10A2_UINT:
-	case DXGI_FORMAT_R11G11B10_FLOAT:
-	case DXGI_FORMAT_R8G8B8A8_TYPELESS:
-	case DXGI_FORMAT_R8G8B8A8_UNORM:
-	case DXGI_FORMAT_R8G8B8A8_UNORM_SRGB:
-	case DXGI_FORMAT_R8G8B8A8_UINT:
-	case DXGI_FORMAT_R8G8B8A8_SNORM:
-	case DXGI_FORMAT_R8G8B8A8_SINT:
-	case DXGI_FORMAT_R16G16_TYPELESS:
-	case DXGI_FORMAT_R16G16_FLOAT:
-	case DXGI_FORMAT_R16G16_UNORM:
-	case DXGI_FORMAT_R16G16_UINT:
-	case DXGI_FORMAT_R16G16_SNORM:
-	case DXGI_FORMAT_R16G16_SINT:
-	case DXGI_FORMAT_R32_TYPELESS:
-	case DXGI_FORMAT_D32_FLOAT:
-	case DXGI_FORMAT_R32_FLOAT:
-	case DXGI_FORMAT_R32_UINT:
-	case DXGI_FORMAT_R32_SINT:
-	case DXGI_FORMAT_R24G8_TYPELESS:
-	case DXGI_FORMAT_D24_UNORM_S8_UINT:
-	case DXGI_FORMAT_R24_UNORM_X8_TYPELESS:
-	case DXGI_FORMAT_X24_TYPELESS_G8_UINT:
-	case DXGI_FORMAT_R9G9B9E5_SHAREDEXP:
-	case DXGI_FORMAT_R8G8_B8G8_UNORM:
-	case DXGI_FORMAT_G8R8_G8B8_UNORM:
-	case DXGI_FORMAT_B8G8R8A8_UNORM:
-	case DXGI_FORMAT_B8G8R8X8_UNORM:
-	case DXGI_FORMAT_R10G10B10_XR_BIAS_A2_UNORM:
-	case DXGI_FORMAT_B8G8R8A8_TYPELESS:
-	case DXGI_FORMAT_B8G8R8A8_UNORM_SRGB:
-	case DXGI_FORMAT_B8G8R8X8_TYPELESS:
-	case DXGI_FORMAT_B8G8R8X8_UNORM_SRGB:
-	case DXGI_FORMAT_AYUV:
-	case DXGI_FORMAT_Y410:
-	case DXGI_FORMAT_P010:
-	case DXGI_FORMAT_P016:
-	case DXGI_FORMAT_YUY2:
-		return 32;
-
-	case DXGI_FORMAT_R8G8_TYPELESS:
-	case DXGI_FORMAT_R8G8_UNORM:
-	case DXGI_FORMAT_R8G8_UINT:
-	case DXGI_FORMAT_R8G8_SNORM:
-	case DXGI_FORMAT_R8G8_SINT:
-	case DXGI_FORMAT_R16_TYPELESS:
-	case DXGI_FORMAT_R16_FLOAT:
-	case DXGI_FORMAT_D16_UNORM:
-	case DXGI_FORMAT_R16_UNORM:
-	case DXGI_FORMAT_R16_UINT:
-	case DXGI_FORMAT_R16_SNORM:
-	case DXGI_FORMAT_R16_SINT:
-	case DXGI_FORMAT_B5G6R5_UNORM:
-	case DXGI_FORMAT_B5G5R5A1_UNORM:
-	case DXGI_FORMAT_B4G4R4A4_UNORM:
-	case DXGI_FORMAT_NV12:
-	case DXGI_FORMAT_NV11:
-		return 16;
-
-	case DXGI_FORMAT_R8_TYPELESS:
-	case DXGI_FORMAT_R8_UNORM:
-	case DXGI_FORMAT_R8_UINT:
-	case DXGI_FORMAT_R8_SNORM:
-	case DXGI_FORMAT_R8_SINT:
-	case DXGI_FORMAT_A8_UNORM:
-		return 8;
-
-	case DXGI_FORMAT_R1_UNORM:
-		return 1;
-	}
-}
-
-static uint32 D3D12GetFormatSizeInBytes(DXGI_FORMAT Format)
-{
-	return D3D12GetFormatSizeInBits(Format) / 8;
-}
-
 struct FD3D12ScopeLock
 {
 public:
diff --git a/Engine/Source/Runtime/D3D12RHI/Private/D3D12View.cpp b/Engine/Source/Runtime/D3D12RHI/Private/D3D12View.cpp
--- a/Engine/Source/Runtime/D3D12RHI/Private/D3D12View.cpp
+++ b/Engine/Source/Runtime/D3D12RHI/Private/D3D12View.cpp
@@ -16,8 +16,8 @@ FD3D12ViewRange::FD3D12ViewRange(D3D12_CONSTANT_BUFFER_VIEW_DESC const& ViewDesc
 
 FD3D12ViewRange::FD3D12ViewRange(D3D12_RENDER_TARGET_VIEW_DESC const& ViewDesc)
 {
-	uint32 FirstPlane = GetPlaneSliceFromViewFormat(ViewDesc.Format, ViewDesc.Format);
-	uint32 NumPlanes  = GetPlaneCount(ViewDesc.Format);
+	uint32 FirstPlane = UE::DXGIUtilities::GetPlaneSliceFromViewFormat(ViewDesc.Format, ViewDesc.Format);
+	uint32 NumPlanes  = UE::DXGIUtilities::GetPlaneCount(ViewDesc.Format);
 	Plane = { FirstPlane, NumPlanes - FirstPlane };
 
 	switch (ViewDesc.ViewDimension)
@@ -61,8 +61,8 @@ FD3D12ViewRange::FD3D12ViewRange(D3D12_RENDER_TARGET_VIEW_DESC const& ViewDesc)
 
 FD3D12ViewRange::FD3D12ViewRange(D3D12_SHADER_RESOURCE_VIEW_DESC const& ViewDesc)
 {
-	uint32 FirstPlane = GetPlaneSliceFromViewFormat(ViewDesc.Format, ViewDesc.Format);
-	uint32 NumPlanes  = GetPlaneCount(ViewDesc.Format);
+	uint32 FirstPlane = UE::DXGIUtilities::GetPlaneSliceFromViewFormat(ViewDesc.Format, ViewDesc.Format);
+	uint32 NumPlanes  = UE::DXGIUtilities::GetPlaneCount(ViewDesc.Format);
 	Plane = { FirstPlane, NumPlanes - FirstPlane };
 
 	switch (ViewDesc.ViewDimension)
@@ -128,8 +128,8 @@ FD3D12ViewRange::FD3D12ViewRange(D3D12_SHADER_RESOURCE_VIEW_DESC const& ViewDesc
 
 FD3D12ViewRange::FD3D12ViewRange(D3D12_DEPTH_STENCIL_VIEW_DESC const& ViewDesc)
 {
-	uint32 FirstPlane = GetPlaneSliceFromViewFormat(ViewDesc.Format, ViewDesc.Format);
-	uint32 NumPlanes  = GetPlaneCount(ViewDesc.Format);
+	uint32 FirstPlane = UE::DXGIUtilities::GetPlaneSliceFromViewFormat(ViewDesc.Format, ViewDesc.Format);
+	uint32 NumPlanes  = UE::DXGIUtilities::GetPlaneCount(ViewDesc.Format);
 	Plane = { FirstPlane, NumPlanes - FirstPlane };
 
 	switch (ViewDesc.ViewDimension)
@@ -166,8 +166,8 @@ FD3D12ViewRange::FD3D12ViewRange(D3D12_DEPTH_STENCIL_VIEW_DESC const& ViewDesc)
 
 FD3D12ViewRange::FD3D12ViewRange(D3D12_UNORDERED_ACCESS_VIEW_DESC const& ViewDesc)
 {
-	uint32 FirstPlane = GetPlaneSliceFromViewFormat(ViewDesc.Format, ViewDesc.Format);
-	uint32 NumPlanes  = GetPlaneCount(ViewDesc.Format);
+	uint32 FirstPlane = UE::DXGIUtilities::GetPlaneSliceFromViewFormat(ViewDesc.Format, ViewDesc.Format);
+	uint32 NumPlanes  = UE::DXGIUtilities::GetPlaneCount(ViewDesc.Format);
 	Plane = { FirstPlane, NumPlanes - FirstPlane };
 
 	switch (ViewDesc.ViewDimension)
diff --git a/Engine/Source/Runtime/D3D12RHI/Private/D3D12View.h b/Engine/Source/Runtime/D3D12RHI/Private/D3D12View.h
--- a/Engine/Source/Runtime/D3D12RHI/Private/D3D12View.h
+++ b/Engine/Source/Runtime/D3D12RHI/Private/D3D12View.h
@@ -53,7 +53,7 @@ struct FD3D12ResourceLayout
 	FD3D12ResourceLayout() = default;
 	FD3D12ResourceLayout(FD3D12ResourceDesc const& ResourceDesc)
 		: NumArraySlices(ResourceDesc.DepthOrArraySize)
-		, NumPlanes     (GetPlaneCount(ResourceDesc.Format))
+		, NumPlanes     (UE::DXGIUtilities::GetPlaneCount(ResourceDesc.Format))
 		, NumMips       (ResourceDesc.MipLevels)
 	{}
 
diff --git a/Engine/Source/Runtime/RHICore/Internal/DXGIUtilities.h b/Engine/Source/Runtime/RHICore/Internal/DXGIUtilities.h
new file mode 100644
--- /dev/null
+++ b/Engine/Source/Runtime/RHICore/Internal/DXGIUtilities.h
@@ -0,0 +1,463 @@
+// Copyright Epic Games, Inc. All Rights Reserved.
+
+#pragma once
+
+#include "HAL/Platform.h"
+
+#if PLATFORM_MICROSOFT
+
+#ifndef RHICORE_PLATFORM_DXGI_H
+	#error "Platform needs to define RHICORE_PLATFORM_DXGI_H"
+#endif
+
+#include "RHIDefinitions.h"
+
+#include "Microsoft/AllowMicrosoftPlatformTypes.h"
+THIRD_PARTY_INCLUDES_START
+	#include RHICORE_PLATFORM_DXGI_H
+THIRD_PARTY_INCLUDES_END
+#include "Microsoft/HideMicrosoftPlatformTypes.h"
+
+namespace UE::DXGIUtilities
+{
+	RHICORE_API const TCHAR* GetFormatString(DXGI_FORMAT Format);
+
+	inline DXGI_FORMAT FindSharedResourceFormat(DXGI_FORMAT InFormat, bool bSRGB)
+	{
+		if (bSRGB)
+		{
+			switch (InFormat)
+			{
+			case DXGI_FORMAT_B8G8R8X8_TYPELESS: return DXGI_FORMAT_B8G8R8X8_UNORM_SRGB;
+			case DXGI_FORMAT_B8G8R8A8_TYPELESS: return DXGI_FORMAT_B8G8R8A8_UNORM_SRGB;
+			case DXGI_FORMAT_R8G8B8A8_TYPELESS: return DXGI_FORMAT_R8G8B8A8_UNORM_SRGB;
+			case DXGI_FORMAT_BC1_TYPELESS:      return DXGI_FORMAT_BC1_UNORM_SRGB;
+			case DXGI_FORMAT_BC2_TYPELESS:      return DXGI_FORMAT_BC2_UNORM_SRGB;
+			case DXGI_FORMAT_BC3_TYPELESS:      return DXGI_FORMAT_BC3_UNORM_SRGB;
+			case DXGI_FORMAT_BC7_TYPELESS:      return DXGI_FORMAT_BC7_UNORM_SRGB;
+			};
+		}
+		else
+		{
+			switch (InFormat)
+			{
+			case DXGI_FORMAT_B8G8R8X8_TYPELESS: return DXGI_FORMAT_B8G8R8X8_UNORM;
+			case DXGI_FORMAT_B8G8R8A8_TYPELESS: return DXGI_FORMAT_B8G8R8A8_UNORM;
+			case DXGI_FORMAT_R8G8B8A8_TYPELESS: return DXGI_FORMAT_R8G8B8A8_UNORM;
+			case DXGI_FORMAT_BC1_TYPELESS:      return DXGI_FORMAT_BC1_UNORM;
+			case DXGI_FORMAT_BC2_TYPELESS:      return DXGI_FORMAT_BC2_UNORM;
+			case DXGI_FORMAT_BC3_TYPELESS:      return DXGI_FORMAT_BC3_UNORM;
+			case DXGI_FORMAT_BC7_TYPELESS:      return DXGI_FORMAT_BC7_UNORM;
+			};
+		}
+
+		switch (InFormat)
+		{
+		case DXGI_FORMAT_R32G32B32A32_TYPELESS: return DXGI_FORMAT_R32G32B32A32_UINT;
+		case DXGI_FORMAT_R32G32B32_TYPELESS:    return DXGI_FORMAT_R32G32B32_UINT;
+		case DXGI_FORMAT_R16G16B16A16_TYPELESS: return DXGI_FORMAT_R16G16B16A16_UNORM;
+		case DXGI_FORMAT_R32G32_TYPELESS:       return DXGI_FORMAT_R32G32_UINT;
+		case DXGI_FORMAT_R10G10B10A2_TYPELESS:  return DXGI_FORMAT_R10G10B10A2_UNORM;
+		case DXGI_FORMAT_R16G16_TYPELESS:       return DXGI_FORMAT_R16G16_UNORM;
+		case DXGI_FORMAT_R8G8_TYPELESS:         return DXGI_FORMAT_R8G8_UNORM;
+		case DXGI_FORMAT_R8_TYPELESS:           return DXGI_FORMAT_R8_UNORM;
+
+		case DXGI_FORMAT_BC4_TYPELESS:          return DXGI_FORMAT_BC4_UNORM;
+		case DXGI_FORMAT_BC5_TYPELESS:          return DXGI_FORMAT_BC5_UNORM;
+		case DXGI_FORMAT_R24G8_TYPELESS:        return DXGI_FORMAT_R24_UNORM_X8_TYPELESS;
+		case DXGI_FORMAT_R32_TYPELESS:          return DXGI_FORMAT_R32_FLOAT;
+		case DXGI_FORMAT_R16_TYPELESS:          return DXGI_FORMAT_R16_UNORM;
+
+		case DXGI_FORMAT_R32G8X24_TYPELESS:     return DXGI_FORMAT_R32_FLOAT_X8X24_TYPELESS;
+		}
+
+		return InFormat;
+	}
+
+	/** Find an appropriate DXGI format for the input format and SRGB setting. */
+	inline DXGI_FORMAT FindShaderResourceFormat(DXGI_FORMAT InFormat, bool bSRGB)
+	{
+		if (bSRGB)
+		{
+			switch (InFormat)
+			{
+			case DXGI_FORMAT_B8G8R8A8_TYPELESS:    return DXGI_FORMAT_B8G8R8A8_UNORM_SRGB;
+			case DXGI_FORMAT_R8G8B8A8_TYPELESS:    return DXGI_FORMAT_R8G8B8A8_UNORM_SRGB;
+			case DXGI_FORMAT_BC1_TYPELESS:         return DXGI_FORMAT_BC1_UNORM_SRGB;
+			case DXGI_FORMAT_BC2_TYPELESS:         return DXGI_FORMAT_BC2_UNORM_SRGB;
+			case DXGI_FORMAT_BC3_TYPELESS:         return DXGI_FORMAT_BC3_UNORM_SRGB;
+			case DXGI_FORMAT_BC7_TYPELESS:         return DXGI_FORMAT_BC7_UNORM_SRGB;
+			};
+		}
+		else
+		{
+			switch (InFormat)
+			{
+			case DXGI_FORMAT_B8G8R8A8_TYPELESS: return DXGI_FORMAT_B8G8R8A8_UNORM;
+			case DXGI_FORMAT_R8G8B8A8_TYPELESS: return DXGI_FORMAT_R8G8B8A8_UNORM;
+			case DXGI_FORMAT_BC1_TYPELESS:      return DXGI_FORMAT_BC1_UNORM;
+			case DXGI_FORMAT_BC2_TYPELESS:      return DXGI_FORMAT_BC2_UNORM;
+			case DXGI_FORMAT_BC3_TYPELESS:      return DXGI_FORMAT_BC3_UNORM;
+			case DXGI_FORMAT_BC7_TYPELESS:      return DXGI_FORMAT_BC7_UNORM;
+			};
+		}
+		switch (InFormat)
+		{
+		case DXGI_FORMAT_R24G8_TYPELESS:    return DXGI_FORMAT_R24_UNORM_X8_TYPELESS;
+		case DXGI_FORMAT_R32_TYPELESS:      return DXGI_FORMAT_R32_FLOAT;
+		case DXGI_FORMAT_R16_TYPELESS:      return DXGI_FORMAT_R16_UNORM;
+		case DXGI_FORMAT_R32G8X24_TYPELESS: return DXGI_FORMAT_R32_FLOAT_X8X24_TYPELESS;
+		}
+		return InFormat;
+	}
+
+	inline DXGI_FORMAT FindDepthStencilResourceFormat(DXGI_FORMAT InFormat)
+	{
+		switch (InFormat)
+		{
+		case DXGI_FORMAT_R32_FLOAT: return DXGI_FORMAT_R32_TYPELESS;
+		case DXGI_FORMAT_R16_FLOAT: return DXGI_FORMAT_R16_TYPELESS;
+		}
+
+		return InFormat;
+	}
+
+	inline DXGI_FORMAT GetPlatformTextureResourceFormat(DXGI_FORMAT InFormat, ETextureCreateFlags InFlags)
+	{
+		// Find valid shared texture format
+		if (EnumHasAnyFlags(InFlags, ETextureCreateFlags::Shared))
+		{
+			return FindSharedResourceFormat(InFormat, EnumHasAnyFlags(InFlags, ETextureCreateFlags::SRGB));
+		}
+		if (EnumHasAnyFlags(InFlags, ETextureCreateFlags::DepthStencilTargetable))
+		{
+			return FindDepthStencilResourceFormat(InFormat);
+		}
+
+		return InFormat;
+	}
+
+	/** Find an appropriate DXGI format unordered access of the raw format. */
+	inline DXGI_FORMAT FindUnorderedAccessFormat(DXGI_FORMAT InFormat)
+	{
+		switch (InFormat)
+		{
+		case DXGI_FORMAT_B8G8R8A8_TYPELESS: return DXGI_FORMAT_B8G8R8A8_UNORM;
+		case DXGI_FORMAT_R8G8B8A8_TYPELESS: return DXGI_FORMAT_R8G8B8A8_UNORM;
+		case DXGI_FORMAT_R32G8X24_TYPELESS: return DXGI_FORMAT_R32_FLOAT_X8X24_TYPELESS;
+		}
+		return InFormat;
+	}
+
+	/** Find the appropriate depth-stencil targetable DXGI format for the given format. */
+	inline DXGI_FORMAT FindDepthStencilFormat(DXGI_FORMAT InFormat)
+	{
+		switch (InFormat)
+		{
+		case DXGI_FORMAT_R24G8_TYPELESS:    return DXGI_FORMAT_D24_UNORM_S8_UINT;
+		case DXGI_FORMAT_R32G8X24_TYPELESS: return DXGI_FORMAT_D32_FLOAT_S8X24_UINT;
+		case DXGI_FORMAT_R32_TYPELESS:      return DXGI_FORMAT_D32_FLOAT;
+		case DXGI_FORMAT_R16_TYPELESS:      return DXGI_FORMAT_D16_UNORM;
+		};
+		return InFormat;
+	}
+
+	/** Find the appropriate depth-stencil typeless DXGI format for the given format. */
+	inline DXGI_FORMAT FindDepthStencilParentDXGIFormat(DXGI_FORMAT InFormat)
+	{
+		switch (InFormat)
+		{
+		case DXGI_FORMAT_D24_UNORM_S8_UINT:			return DXGI_FORMAT_R24G8_TYPELESS;
+		case DXGI_FORMAT_X24_TYPELESS_G8_UINT:		return DXGI_FORMAT_R24G8_TYPELESS;
+
+		case DXGI_FORMAT_D32_FLOAT_S8X24_UINT:		return DXGI_FORMAT_R32G8X24_TYPELESS;
+		case DXGI_FORMAT_X32_TYPELESS_G8X24_UINT:	return DXGI_FORMAT_R32G8X24_TYPELESS;
+
+		case DXGI_FORMAT_D32_FLOAT:					return DXGI_FORMAT_R32_TYPELESS;
+		case DXGI_FORMAT_D16_UNORM:					return DXGI_FORMAT_R16_TYPELESS;
+		};
+		return InFormat;
+	}
+
+	/**
+	 * Returns whether the given format contains stencil information.
+	 * Must be passed a format returned by FindDepthStencilFormat, so that typeless versions are converted to their corresponding depth stencil view format.
+	 */
+	inline bool HasStencilBits(DXGI_FORMAT InFormat)
+	{
+		switch (InFormat)
+		{
+		case DXGI_FORMAT_D24_UNORM_S8_UINT:
+		case DXGI_FORMAT_D32_FLOAT_S8X24_UINT:
+			return true;
+		};
+		return false;
+	}
+
+	inline bool IsBlockCompressedFormat(DXGI_FORMAT Format)
+	{
+		// Returns true if BC1, BC2, BC3, BC4, BC5, BC6, BC7
+		return (Format >= DXGI_FORMAT_BC1_TYPELESS && Format <= DXGI_FORMAT_BC5_SNORM) ||
+			(Format >= DXGI_FORMAT_BC6H_TYPELESS && Format <= DXGI_FORMAT_BC7_UNORM_SRGB);
+	}
+
+	inline uint8 GetPlaneSliceFromViewFormat(DXGI_FORMAT ResourceFormat, DXGI_FORMAT ViewFormat)
+	{
+		// Currently, the only planar resources used are depth-stencil formats
+		switch (FindDepthStencilParentDXGIFormat(ResourceFormat))
+		{
+		case DXGI_FORMAT_R24G8_TYPELESS:
+			switch (ViewFormat)
+			{
+			case DXGI_FORMAT_R24_UNORM_X8_TYPELESS: return 0;
+			case DXGI_FORMAT_X24_TYPELESS_G8_UINT: return 1;
+			}
+			break;
+
+		case DXGI_FORMAT_R32G8X24_TYPELESS:
+			switch (ViewFormat)
+			{
+			case DXGI_FORMAT_R32_FLOAT_X8X24_TYPELESS: return 0;
+			case DXGI_FORMAT_X32_TYPELESS_G8X24_UINT: return 1;
+			}
+			break;
+
+		case DXGI_FORMAT_NV12:
+			switch (ViewFormat)
+			{
+			case DXGI_FORMAT_R8_UNORM: return 0;
+			case DXGI_FORMAT_R8G8_UNORM: return 1;
+			}
+			break;
+		}
+
+		return 0;
+	}
+
+	inline uint8 GetPlaneCount(DXGI_FORMAT Format)
+	{
+		// Currently, the only planar resources used are depth-stencil formats
+		switch (FindDepthStencilParentDXGIFormat(Format))
+		{
+		case DXGI_FORMAT_R24G8_TYPELESS:
+		case DXGI_FORMAT_R32G8X24_TYPELESS:
+			return 2;
+		default:
+			return 1;
+		}
+	}
+
+	inline uint32 GetWidthAlignment(DXGI_FORMAT Format)
+	{
+		switch (Format)
+		{
+		case DXGI_FORMAT_R8G8_B8G8_UNORM: return 2;
+		case DXGI_FORMAT_G8R8_G8B8_UNORM: return 2;
+		case DXGI_FORMAT_NV12: return 2;
+		case DXGI_FORMAT_P010: return 2;
+		case DXGI_FORMAT_P016: return 2;
+		case DXGI_FORMAT_420_OPAQUE: return 2;
+		case DXGI_FORMAT_YUY2: return 2;
+		case DXGI_FORMAT_Y210: return 2;
+		case DXGI_FORMAT_Y216: return 2;
+		case DXGI_FORMAT_BC1_TYPELESS: return 4;
+		case DXGI_FORMAT_BC1_UNORM: return 4;
+		case DXGI_FORMAT_BC1_UNORM_SRGB: return 4;
+		case DXGI_FORMAT_BC2_TYPELESS: return 4;
+		case DXGI_FORMAT_BC2_UNORM: return 4;
+		case DXGI_FORMAT_BC2_UNORM_SRGB: return 4;
+		case DXGI_FORMAT_BC3_TYPELESS: return 4;
+		case DXGI_FORMAT_BC3_UNORM: return 4;
+		case DXGI_FORMAT_BC3_UNORM_SRGB: return 4;
+		case DXGI_FORMAT_BC4_TYPELESS: return 4;
+		case DXGI_FORMAT_BC4_UNORM: return 4;
+		case DXGI_FORMAT_BC4_SNORM: return 4;
+		case DXGI_FORMAT_BC5_TYPELESS: return 4;
+		case DXGI_FORMAT_BC5_UNORM: return 4;
+		case DXGI_FORMAT_BC5_SNORM: return 4;
+		case DXGI_FORMAT_BC6H_TYPELESS: return 4;
+		case DXGI_FORMAT_BC6H_UF16: return 4;
+		case DXGI_FORMAT_BC6H_SF16: return 4;
+		case DXGI_FORMAT_BC7_TYPELESS: return 4;
+		case DXGI_FORMAT_BC7_UNORM: return 4;
+		case DXGI_FORMAT_BC7_UNORM_SRGB: return 4;
+		case DXGI_FORMAT_NV11: return 4;
+		case DXGI_FORMAT_R1_UNORM: return 8;
+		default: return 1;
+		}
+	}
+
+	inline uint32 GetHeightAlignment(DXGI_FORMAT Format)
+	{
+		switch (Format)
+		{
+		case DXGI_FORMAT_NV12: return 2;
+		case DXGI_FORMAT_P010: return 2;
+		case DXGI_FORMAT_P016: return 2;
+		case DXGI_FORMAT_420_OPAQUE: return 2;
+		case DXGI_FORMAT_BC1_TYPELESS: return 4;
+		case DXGI_FORMAT_BC1_UNORM: return 4;
+		case DXGI_FORMAT_BC1_UNORM_SRGB: return 4;
+		case DXGI_FORMAT_BC2_TYPELESS: return 4;
+		case DXGI_FORMAT_BC2_UNORM: return 4;
+		case DXGI_FORMAT_BC2_UNORM_SRGB: return 4;
+		case DXGI_FORMAT_BC3_TYPELESS: return 4;
+		case DXGI_FORMAT_BC3_UNORM: return 4;
+		case DXGI_FORMAT_BC3_UNORM_SRGB: return 4;
+		case DXGI_FORMAT_BC4_TYPELESS: return 4;
+		case DXGI_FORMAT_BC4_UNORM: return 4;
+		case DXGI_FORMAT_BC4_SNORM: return 4;
+		case DXGI_FORMAT_BC5_TYPELESS: return 4;
+		case DXGI_FORMAT_BC5_UNORM: return 4;
+		case DXGI_FORMAT_BC5_SNORM: return 4;
+		case DXGI_FORMAT_BC6H_TYPELESS: return 4;
+		case DXGI_FORMAT_BC6H_UF16: return 4;
+		case DXGI_FORMAT_BC6H_SF16: return 4;
+		case DXGI_FORMAT_BC7_TYPELESS: return 4;
+		case DXGI_FORMAT_BC7_UNORM: return 4;
+		case DXGI_FORMAT_BC7_UNORM_SRGB: return 4;
+		default: return 1;
+		}
+	}
+
+	inline uint32 GetFormatSizeInBits(DXGI_FORMAT Format)
+	{
+		switch (Format)
+		{
+		default: checkNoEntry(); [[fallthrough]];
+		case DXGI_FORMAT_UNKNOWN:
+			return 0;
+
+		case DXGI_FORMAT_R32G32B32A32_TYPELESS:
+		case DXGI_FORMAT_R32G32B32A32_FLOAT:
+		case DXGI_FORMAT_R32G32B32A32_UINT:
+		case DXGI_FORMAT_R32G32B32A32_SINT:
+		case DXGI_FORMAT_BC2_TYPELESS:
+		case DXGI_FORMAT_BC2_UNORM:
+		case DXGI_FORMAT_BC2_UNORM_SRGB:
+		case DXGI_FORMAT_BC3_TYPELESS:
+		case DXGI_FORMAT_BC3_UNORM:
+		case DXGI_FORMAT_BC3_UNORM_SRGB:
+		case DXGI_FORMAT_BC5_TYPELESS:
+		case DXGI_FORMAT_BC5_UNORM:
+		case DXGI_FORMAT_BC5_SNORM:
+		case DXGI_FORMAT_BC6H_TYPELESS:
+		case DXGI_FORMAT_BC6H_UF16:
+		case DXGI_FORMAT_BC6H_SF16:
+		case DXGI_FORMAT_BC7_TYPELESS:
+		case DXGI_FORMAT_BC7_UNORM:
+		case DXGI_FORMAT_BC7_UNORM_SRGB:
+			return 128;
+
+		case DXGI_FORMAT_R32G32B32_TYPELESS:
+		case DXGI_FORMAT_R32G32B32_FLOAT:
+		case DXGI_FORMAT_R32G32B32_UINT:
+		case DXGI_FORMAT_R32G32B32_SINT:
+			return 96;
+
+		case DXGI_FORMAT_R16G16B16A16_TYPELESS:
+		case DXGI_FORMAT_R16G16B16A16_FLOAT:
+		case DXGI_FORMAT_R16G16B16A16_UNORM:
+		case DXGI_FORMAT_R16G16B16A16_UINT:
+		case DXGI_FORMAT_R16G16B16A16_SNORM:
+		case DXGI_FORMAT_R16G16B16A16_SINT:
+		case DXGI_FORMAT_R32G32_TYPELESS:
+		case DXGI_FORMAT_R32G32_FLOAT:
+		case DXGI_FORMAT_R32G32_UINT:
+		case DXGI_FORMAT_R32G32_SINT:
+		case DXGI_FORMAT_R32G8X24_TYPELESS:
+		case DXGI_FORMAT_D32_FLOAT_S8X24_UINT:
+		case DXGI_FORMAT_R32_FLOAT_X8X24_TYPELESS:
+		case DXGI_FORMAT_X32_TYPELESS_G8X24_UINT:
+		case DXGI_FORMAT_BC1_TYPELESS:
+		case DXGI_FORMAT_BC1_UNORM:
+		case DXGI_FORMAT_BC1_UNORM_SRGB:
+		case DXGI_FORMAT_BC4_TYPELESS:
+		case DXGI_FORMAT_BC4_UNORM:
+		case DXGI_FORMAT_BC4_SNORM:
+		case DXGI_FORMAT_Y416:
+		case DXGI_FORMAT_Y210:
+		case DXGI_FORMAT_Y216:
+			return 64;
+
+		case DXGI_FORMAT_R10G10B10A2_TYPELESS:
+		case DXGI_FORMAT_R10G10B10A2_UNORM:
+		case DXGI_FORMAT_R10G10B10A2_UINT:
+		case DXGI_FORMAT_R11G11B10_FLOAT:
+		case DXGI_FORMAT_R8G8B8A8_TYPELESS:
+		case DXGI_FORMAT_R8G8B8A8_UNORM:
+		case DXGI_FORMAT_R8G8B8A8_UNORM_SRGB:
+		case DXGI_FORMAT_R8G8B8A8_UINT:
+		case DXGI_FORMAT_R8G8B8A8_SNORM:
+		case DXGI_FORMAT_R8G8B8A8_SINT:
+		case DXGI_FORMAT_R16G16_TYPELESS:
+		case DXGI_FORMAT_R16G16_FLOAT:
+		case DXGI_FORMAT_R16G16_UNORM:
+		case DXGI_FORMAT_R16G16_UINT:
+		case DXGI_FORMAT_R16G16_SNORM:
+		case DXGI_FORMAT_R16G16_SINT:
+		case DXGI_FORMAT_R32_TYPELESS:
+		case DXGI_FORMAT_D32_FLOAT:
+		case DXGI_FORMAT_R32_FLOAT:
+		case DXGI_FORMAT_R32_UINT:
+		case DXGI_FORMAT_R32_SINT:
+		case DXGI_FORMAT_R24G8_TYPELESS:
+		case DXGI_FORMAT_D24_UNORM_S8_UINT:
+		case DXGI_FORMAT_R24_UNORM_X8_TYPELESS:
+		case DXGI_FORMAT_X24_TYPELESS_G8_UINT:
+		case DXGI_FORMAT_R9G9B9E5_SHAREDEXP:
+		case DXGI_FORMAT_R8G8_B8G8_UNORM:
+		case DXGI_FORMAT_G8R8_G8B8_UNORM:
+		case DXGI_FORMAT_B8G8R8A8_UNORM:
+		case DXGI_FORMAT_B8G8R8X8_UNORM:
+		case DXGI_FORMAT_R10G10B10_XR_BIAS_A2_UNORM:
+		case DXGI_FORMAT_B8G8R8A8_TYPELESS:
+		case DXGI_FORMAT_B8G8R8A8_UNORM_SRGB:
+		case DXGI_FORMAT_B8G8R8X8_TYPELESS:
+		case DXGI_FORMAT_B8G8R8X8_UNORM_SRGB:
+		case DXGI_FORMAT_AYUV:
+		case DXGI_FORMAT_Y410:
+		case DXGI_FORMAT_P010:
+		case DXGI_FORMAT_P016:
+		case DXGI_FORMAT_YUY2:
+			return 32;
+
+		case DXGI_FORMAT_R8G8_TYPELESS:
+		case DXGI_FORMAT_R8G8_UNORM:
+		case DXGI_FORMAT_R8G8_UINT:
+		case DXGI_FORMAT_R8G8_SNORM:
+		case DXGI_FORMAT_R8G8_SINT:
+		case DXGI_FORMAT_R16_TYPELESS:
+		case DXGI_FORMAT_R16_FLOAT:
+		case DXGI_FORMAT_D16_UNORM:
+		case DXGI_FORMAT_R16_UNORM:
+		case DXGI_FORMAT_R16_UINT:
+		case DXGI_FORMAT_R16_SNORM:
+		case DXGI_FORMAT_R16_SINT:
+		case DXGI_FORMAT_B5G6R5_UNORM:
+		case DXGI_FORMAT_B5G5R5A1_UNORM:
+		case DXGI_FORMAT_B4G4R4A4_UNORM:
+		case DXGI_FORMAT_NV12:
+		case DXGI_FORMAT_NV11:
+			return 16;
+
+		case DXGI_FORMAT_R8_TYPELESS:
+		case DXGI_FORMAT_R8_UNORM:
+		case DXGI_FORMAT_R8_UINT:
+		case DXGI_FORMAT_R8_SNORM:
+		case DXGI_FORMAT_R8_SINT:
+		case DXGI_FORMAT_A8_UNORM:
+			return 8;
+
+		case DXGI_FORMAT_R1_UNORM:
+			return 1;
+		}
+	}
+
+	inline uint32 GetFormatSizeInBytes(DXGI_FORMAT Format)
+	{
+		return GetFormatSizeInBits(Format) / 8;
+	}
+} // UE::DXGIUtilities
+
+#endif // PLATFORM_MICROSOFT
diff --git a/Engine/Source/Runtime/RHICore/Private/DXGIUtilities.cpp b/Engine/Source/Runtime/RHICore/Private/DXGIUtilities.cpp
new file mode 100644
--- /dev/null
+++ b/Engine/Source/Runtime/RHICore/Private/DXGIUtilities.cpp
@@ -0,0 +1,52 @@
+// Copyright Epic Games, Inc. All Rights Reserved.
+
+#include "DXGIUtilities.h"
+
+#if PLATFORM_MICROSOFT
+
+const TCHAR* UE::DXGIUtilities::GetFormatString(DXGI_FORMAT Format)
+{
+	const TCHAR* Result = TEXT("");
+
+#define DXGI_FORMAT_CASE(x) case x: Result = TEXT(#x); break;
+	switch (Format)
+	{
+		DXGI_FORMAT_CASE(DXGI_FORMAT_R8G8B8A8_UNORM)
+		DXGI_FORMAT_CASE(DXGI_FORMAT_B8G8R8A8_UNORM)
+		DXGI_FORMAT_CASE(DXGI_FORMAT_B8G8R8X8_UNORM)
+		DXGI_FORMAT_CASE(DXGI_FORMAT_BC1_UNORM)
+		DXGI_FORMAT_CASE(DXGI_FORMAT_BC2_UNORM)
+		DXGI_FORMAT_CASE(DXGI_FORMAT_BC3_UNORM)
+		DXGI_FORMAT_CASE(DXGI_FORMAT_BC4_UNORM)
+		DXGI_FORMAT_CASE(DXGI_FORMAT_R16G16B16A16_FLOAT)
+		DXGI_FORMAT_CASE(DXGI_FORMAT_R32G32B32A32_FLOAT)
+		DXGI_FORMAT_CASE(DXGI_FORMAT_UNKNOWN)
+		DXGI_FORMAT_CASE(DXGI_FORMAT_R8_UNORM)
+		DXGI_FORMAT_CASE(DXGI_FORMAT_D32_FLOAT_S8X24_UINT)
+		DXGI_FORMAT_CASE(DXGI_FORMAT_R32_FLOAT_X8X24_TYPELESS)
+		DXGI_FORMAT_CASE(DXGI_FORMAT_R32G8X24_TYPELESS)
+		DXGI_FORMAT_CASE(DXGI_FORMAT_D24_UNORM_S8_UINT)
+		DXGI_FORMAT_CASE(DXGI_FORMAT_R24_UNORM_X8_TYPELESS)
+		DXGI_FORMAT_CASE(DXGI_FORMAT_R32_FLOAT)
+		DXGI_FORMAT_CASE(DXGI_FORMAT_R16G16_UINT)
+		DXGI_FORMAT_CASE(DXGI_FORMAT_R16G16_UNORM)
+		DXGI_FORMAT_CASE(DXGI_FORMAT_R16G16_SNORM)
+		DXGI_FORMAT_CASE(DXGI_FORMAT_R16G16_FLOAT)
+		DXGI_FORMAT_CASE(DXGI_FORMAT_R32G32_FLOAT)
+		DXGI_FORMAT_CASE(DXGI_FORMAT_R10G10B10A2_UNORM)
+		DXGI_FORMAT_CASE(DXGI_FORMAT_R16G16B16A16_UINT)
+		DXGI_FORMAT_CASE(DXGI_FORMAT_R8G8_SNORM)
+		DXGI_FORMAT_CASE(DXGI_FORMAT_BC5_UNORM)
+		DXGI_FORMAT_CASE(DXGI_FORMAT_R1_UNORM)
+		DXGI_FORMAT_CASE(DXGI_FORMAT_R8G8B8A8_TYPELESS)
+		DXGI_FORMAT_CASE(DXGI_FORMAT_B8G8R8A8_TYPELESS)
+		DXGI_FORMAT_CASE(DXGI_FORMAT_BC7_UNORM)
+		DXGI_FORMAT_CASE(DXGI_FORMAT_BC6H_UF16)
+		default:
+			Result = TEXT("");
+	}
+#undef DXGI_FORMAT_CASE
+	return Result;
+}
+
+#endif // PLATFORM_MICROSOFT
diff --git a/Engine/Source/Runtime/RHICore/RHICore.Build.cs b/Engine/Source/Runtime/RHICore/RHICore.Build.cs
--- a/Engine/Source/Runtime/RHICore/RHICore.Build.cs
+++ b/Engine/Source/Runtime/RHICore/RHICore.Build.cs
@@ -1,7 +1,6 @@
 // Copyright Epic Games, Inc. All Rights Reserved.
 
 using UnrealBuildTool;
-using System;
 
 public class RHICore : ModuleRules
 {
@@ -10,5 +9,10 @@ public class RHICore : ModuleRules
 		PublicDependencyModuleNames.AddRange(new string[] { "RHI" });
 		PublicDependencyModuleNames.AddRange(new string[] { "RenderCore" });
 		PrivateDependencyModuleNames.AddRange(new string[] { "Core", "Projects", "RHI", "ApplicationCore", "TraceLog" });
+
+		if (Target.Platform.IsInGroup(UnrealPlatformGroup.Windows))
+		{
+			PublicDefinitions.Add("RHICORE_PLATFORM_DXGI_H=<dxgi.h>");
+		}
 	}
 }
diff --git a/Engine/Source/Runtime/Windows/D3D11RHI/Private/D3D11Device.cpp b/Engine/Source/Runtime/Windows/D3D11RHI/Private/D3D11Device.cpp
--- a/Engine/Source/Runtime/Windows/D3D11RHI/Private/D3D11Device.cpp
+++ b/Engine/Source/Runtime/Windows/D3D11RHI/Private/D3D11Device.cpp
@@ -432,17 +432,17 @@ void FD3D11DynamicRHI::SetupAfterDeviceCreation()
 	{
 		FPixelFormatInfo& PixelFormatInfo = GPixelFormats[FormatIndex];
 		const DXGI_FORMAT PlatformFormat = static_cast<DXGI_FORMAT>(PixelFormatInfo.PlatformFormat);
-		const DXGI_FORMAT UAVFormat = FindUnorderedAccessDXGIFormat(PlatformFormat);
+		const DXGI_FORMAT UAVFormat = UE::DXGIUtilities::FindUnorderedAccessFormat(PlatformFormat);
 
 		EPixelFormatCapabilities Capabilities = EPixelFormatCapabilities::None;
 
 		if (PlatformFormat != DXGI_FORMAT_UNKNOWN)
 		{
 			const FFormatSupport FormatSupport    = GetFormatSupport(Direct3DDevice, PlatformFormat);
-			const FFormatSupport SRVFormatSupport = GetFormatSupport(Direct3DDevice, FindShaderResourceDXGIFormat(PlatformFormat, false));
-			const FFormatSupport UAVFormatSupport = GetFormatSupport(Direct3DDevice, FindUnorderedAccessDXGIFormat(PlatformFormat));
-			const FFormatSupport RTVFormatSupport = GetFormatSupport(Direct3DDevice, FindShaderResourceDXGIFormat(PlatformFormat, false));
-			const FFormatSupport DSVFormatSupport = GetFormatSupport(Direct3DDevice, FindDepthStencilDXGIFormat(PlatformFormat));
+			const FFormatSupport SRVFormatSupport = GetFormatSupport(Direct3DDevice, UE::DXGIUtilities::FindShaderResourceFormat(PlatformFormat, false));
+			const FFormatSupport UAVFormatSupport = GetFormatSupport(Direct3DDevice, UE::DXGIUtilities::FindUnorderedAccessFormat(PlatformFormat));
+			const FFormatSupport RTVFormatSupport = GetFormatSupport(Direct3DDevice, UE::DXGIUtilities::FindShaderResourceFormat(PlatformFormat, false));
+			const FFormatSupport DSVFormatSupport = GetFormatSupport(Direct3DDevice, UE::DXGIUtilities::FindDepthStencilFormat(PlatformFormat));
 
 			auto ConvertCap1 = [&Capabilities](const FFormatSupport& InSupport, EPixelFormatCapabilities UnrealCap, D3D11_FORMAT_SUPPORT InFlag)
 			{
@@ -747,7 +747,7 @@ static bool CanFormatBeDisplayed(const FD3D11DynamicRHI* InD3DRHI, EPixelFormat
 	const HRESULT FormatSupportResult = InD3DRHI->GetDevice()->CheckFormatSupport(DxgiFormat, &FormatSupport);
 	if (FAILED(FormatSupportResult))
 	{
-		const TCHAR* D3DFormatString = GetD3D11TextureFormatString(DxgiFormat);
+		const TCHAR* D3DFormatString = UE::DXGIUtilities::GetFormatString(DxgiFormat);
 		UE_LOG(LogD3D11RHI, Warning, TEXT("CheckFormatSupport(%s) failed: 0x%08x"), D3DFormatString, FormatSupportResult);
 		return false;
 	}
diff --git a/Engine/Source/Runtime/Windows/D3D11RHI/Private/D3D11RHI.cpp b/Engine/Source/Runtime/Windows/D3D11RHI/Private/D3D11RHI.cpp
--- a/Engine/Source/Runtime/Windows/D3D11RHI/Private/D3D11RHI.cpp
+++ b/Engine/Source/Runtime/Windows/D3D11RHI/Private/D3D11RHI.cpp
@@ -767,8 +767,8 @@ IDXGISwapChain* FD3D11DynamicRHI::RHIGetSwapChain(FRHIViewport* InViewport) cons
 
 DXGI_FORMAT FD3D11DynamicRHI::RHIGetSwapChainFormat(EPixelFormat InFormat) const
 {
-	const DXGI_FORMAT PlatformFormat = ::FindDepthStencilDXGIFormat(static_cast<DXGI_FORMAT>(GPixelFormats[InFormat].PlatformFormat));
-	return ::FindShaderResourceDXGIFormat(PlatformFormat, true);
+	const DXGI_FORMAT PlatformFormat = UE::DXGIUtilities::FindDepthStencilFormat(static_cast<DXGI_FORMAT>(GPixelFormats[InFormat].PlatformFormat));
+	return UE::DXGIUtilities::FindShaderResourceFormat(PlatformFormat, true);
 }
 
 ID3D11Buffer* FD3D11DynamicRHI::RHIGetResource(FRHIBuffer* InBuffer) const
diff --git a/Engine/Source/Runtime/Windows/D3D11RHI/Private/D3D11RHIPrivate.h b/Engine/Source/Runtime/Windows/D3D11RHI/Private/D3D11RHIPrivate.h
--- a/Engine/Source/Runtime/Windows/D3D11RHI/Private/D3D11RHIPrivate.h
+++ b/Engine/Source/Runtime/Windows/D3D11RHI/Private/D3D11RHIPrivate.h
@@ -18,6 +18,7 @@
 #include "Engine/Engine.h"
 #include "HDRHelper.h"
 #include "BoundShaderStateHistory.h"
+#include "DXGIUtilities.h"
 
 DECLARE_LOG_CATEGORY_EXTERN(LogD3D11RHI, Log, All);
 
@@ -767,9 +768,6 @@ public:
 	void ConditionalClearShaderResource(FD3D11ViewableResource* Resource, bool bCheckBoundInputAssembler);
 	void ClearAllShaderResources();
 
-
-	static DXGI_FORMAT GetPlatformTextureResourceFormat(DXGI_FORMAT InFormat, ETextureCreateFlags InFlags);
-
 	uint32 GetHDRDetectedDisplayIndex() const
 	{
 		return HDRDetectedDisplayIndex;
@@ -1080,90 +1078,6 @@ private:
 	void FindAdapter();
 };
 
-/** Find an appropriate DXGI format for the input format and SRGB setting. */
-inline DXGI_FORMAT FindShaderResourceDXGIFormat(DXGI_FORMAT InFormat,bool bSRGB)
-{
-	if(bSRGB)
-	{
-		switch(InFormat)
-		{
-		case DXGI_FORMAT_B8G8R8A8_TYPELESS:    return DXGI_FORMAT_B8G8R8A8_UNORM_SRGB;
-		case DXGI_FORMAT_R8G8B8A8_TYPELESS:    return DXGI_FORMAT_R8G8B8A8_UNORM_SRGB;
-		case DXGI_FORMAT_BC1_TYPELESS:         return DXGI_FORMAT_BC1_UNORM_SRGB;
-		case DXGI_FORMAT_BC2_TYPELESS:         return DXGI_FORMAT_BC2_UNORM_SRGB;
-		case DXGI_FORMAT_BC3_TYPELESS:         return DXGI_FORMAT_BC3_UNORM_SRGB;
-		case DXGI_FORMAT_BC7_TYPELESS:         return DXGI_FORMAT_BC7_UNORM_SRGB;
-		};
-	}
-	else
-	{
-		switch(InFormat)
-		{
-		case DXGI_FORMAT_B8G8R8A8_TYPELESS: return DXGI_FORMAT_B8G8R8A8_UNORM;
-		case DXGI_FORMAT_R8G8B8A8_TYPELESS: return DXGI_FORMAT_R8G8B8A8_UNORM;
-		case DXGI_FORMAT_BC1_TYPELESS:      return DXGI_FORMAT_BC1_UNORM;
-		case DXGI_FORMAT_BC2_TYPELESS:      return DXGI_FORMAT_BC2_UNORM;
-		case DXGI_FORMAT_BC3_TYPELESS:      return DXGI_FORMAT_BC3_UNORM;
-		case DXGI_FORMAT_BC7_TYPELESS:      return DXGI_FORMAT_BC7_UNORM;
-		};
-	}
-	switch(InFormat)
-	{
-		case DXGI_FORMAT_R24G8_TYPELESS: return DXGI_FORMAT_R24_UNORM_X8_TYPELESS;
-		case DXGI_FORMAT_R32_TYPELESS: return DXGI_FORMAT_R32_FLOAT;
-		case DXGI_FORMAT_R16_TYPELESS: return DXGI_FORMAT_R16_UNORM;
-		// Changing Depth Buffers to 32 bit on Dingo as D24S8 is actually implemented as a 32 bit buffer in the hardware
-		case DXGI_FORMAT_R32G8X24_TYPELESS: return DXGI_FORMAT_R32_FLOAT_X8X24_TYPELESS; 
-	}
-	return InFormat;
-}
-
-/** Find an appropriate DXGI format unordered access of the raw format. */
-inline DXGI_FORMAT FindUnorderedAccessDXGIFormat(DXGI_FORMAT InFormat)
-{
-	switch(InFormat)
-	{
-		case DXGI_FORMAT_B8G8R8A8_TYPELESS: return DXGI_FORMAT_B8G8R8A8_UNORM;
-		case DXGI_FORMAT_R8G8B8A8_TYPELESS: return DXGI_FORMAT_R8G8B8A8_UNORM;
-	}
-	return InFormat;
-}
-
-/** Find the appropriate depth-stencil targetable DXGI format for the given format. */
-inline DXGI_FORMAT FindDepthStencilDXGIFormat(DXGI_FORMAT InFormat)
-{
-	switch(InFormat)
-	{
-		case DXGI_FORMAT_R24G8_TYPELESS:
-			return DXGI_FORMAT_D24_UNORM_S8_UINT;
-		// Changing Depth Buffers to 32 bit on Dingo as D24S8 is actually implemented as a 32 bit buffer in the hardware
-		case DXGI_FORMAT_R32G8X24_TYPELESS:
-			return DXGI_FORMAT_D32_FLOAT_S8X24_UINT;
-		case DXGI_FORMAT_R32_TYPELESS:
-			return DXGI_FORMAT_D32_FLOAT;
-		case DXGI_FORMAT_R16_TYPELESS:
-			return DXGI_FORMAT_D16_UNORM;
-	};
-	return InFormat;
-}
-
-/** 
- * Returns whether the given format contains stencil information.  
- * Must be passed a format returned by FindDepthStencilDXGIFormat, so that typeless versions are converted to their corresponding depth stencil view format.
- */
-inline bool HasStencilBits(DXGI_FORMAT InFormat)
-{
-	switch(InFormat)
-	{
-	case DXGI_FORMAT_D24_UNORM_S8_UINT:
-		return true;
-	// Changing Depth Buffers to 32 bit on Dingo as D24S8 is actually implemented as a 32 bit buffer in the hardware
-	case DXGI_FORMAT_D32_FLOAT_S8X24_UINT:
-		return true;
-	};
-	return false;
-}
-
 // 1d, 31 bit (uses the sign bit for internal use), O(n) where n is the amount of elements stored
 // does not enforce any alignment
 // unoccupied regions get compacted but occupied don't get compacted
diff --git a/Engine/Source/Runtime/Windows/D3D11RHI/Private/D3D11SRV.cpp b/Engine/Source/Runtime/Windows/D3D11RHI/Private/D3D11SRV.cpp
--- a/Engine/Source/Runtime/Windows/D3D11RHI/Private/D3D11SRV.cpp
+++ b/Engine/Source/Runtime/Windows/D3D11RHI/Private/D3D11SRV.cpp
@@ -39,7 +39,7 @@ void FD3D11ShaderResourceView::UpdateView()
 		if (!Info.bNullView)
 		{
 			D3DResource = Buffer->Resource;
-			SRVDesc.Format = FindShaderResourceDXGIFormat(DXGI_FORMAT(GPixelFormats[Info.Format].PlatformFormat), false);
+			SRVDesc.Format = UE::DXGIUtilities::FindShaderResourceFormat(DXGI_FORMAT(GPixelFormats[Info.Format].PlatformFormat), false);
 
 			switch (Info.BufferType)
 			{
@@ -72,7 +72,7 @@ void FD3D11ShaderResourceView::UpdateView()
 		FRHITextureDesc const& TextureDesc = Texture->GetDesc();
 		auto const Info = ViewDesc.Texture.SRV.GetViewInfo(Texture);
 
-		SRVDesc.Format = FindShaderResourceDXGIFormat(DXGI_FORMAT(GPixelFormats[Info.Format].PlatformFormat), Info.bSRGB);
+		SRVDesc.Format = UE::DXGIUtilities::FindShaderResourceFormat(DXGI_FORMAT(GPixelFormats[Info.Format].PlatformFormat), Info.bSRGB);
 
 		// No need to use Info.Dimension, since D3D supports mixing Texture2D view types.
 		// Create a view which matches the underlying resource dimension.
diff --git a/Engine/Source/Runtime/Windows/D3D11RHI/Private/D3D11Texture.cpp b/Engine/Source/Runtime/Windows/D3D11RHI/Private/D3D11Texture.cpp
--- a/Engine/Source/Runtime/Windows/D3D11RHI/Private/D3D11Texture.cpp
+++ b/Engine/Source/Runtime/Windows/D3D11RHI/Private/D3D11Texture.cpp
@@ -189,86 +189,6 @@ bool FD3D11DynamicRHI::RHIGetTextureMemoryVisualizeData( FColor* /*TextureData*/
 	return false;
 }
 
-/** Find an appropriate DXGI format for the input format and SRGB setting. */
-inline DXGI_FORMAT FindSharedResourceDXGIFormat(DXGI_FORMAT InFormat, bool bSRGB)
-{
-	if (bSRGB)
-	{
-		switch (InFormat)
-		{
-		case DXGI_FORMAT_B8G8R8X8_TYPELESS:    return DXGI_FORMAT_B8G8R8X8_UNORM_SRGB;
-		case DXGI_FORMAT_B8G8R8A8_TYPELESS:    return DXGI_FORMAT_B8G8R8A8_UNORM_SRGB;
-		case DXGI_FORMAT_R8G8B8A8_TYPELESS:    return DXGI_FORMAT_R8G8B8A8_UNORM_SRGB;
-		case DXGI_FORMAT_BC1_TYPELESS:         return DXGI_FORMAT_BC1_UNORM_SRGB;
-		case DXGI_FORMAT_BC2_TYPELESS:         return DXGI_FORMAT_BC2_UNORM_SRGB;
-		case DXGI_FORMAT_BC3_TYPELESS:         return DXGI_FORMAT_BC3_UNORM_SRGB;
-		case DXGI_FORMAT_BC7_TYPELESS:         return DXGI_FORMAT_BC7_UNORM_SRGB;
-		};
-	}
-	else
-	{
-		switch (InFormat)
-		{
-		case DXGI_FORMAT_B8G8R8X8_TYPELESS:    return DXGI_FORMAT_B8G8R8X8_UNORM;
-		case DXGI_FORMAT_B8G8R8A8_TYPELESS: return DXGI_FORMAT_B8G8R8A8_UNORM;
-		case DXGI_FORMAT_R8G8B8A8_TYPELESS: return DXGI_FORMAT_R8G8B8A8_UNORM;
-		case DXGI_FORMAT_BC1_TYPELESS:      return DXGI_FORMAT_BC1_UNORM;
-		case DXGI_FORMAT_BC2_TYPELESS:      return DXGI_FORMAT_BC2_UNORM;
-		case DXGI_FORMAT_BC3_TYPELESS:      return DXGI_FORMAT_BC3_UNORM;
-		case DXGI_FORMAT_BC7_TYPELESS:      return DXGI_FORMAT_BC7_UNORM;
-		};
-	}
-	switch (InFormat)
-	{
-	case DXGI_FORMAT_R32G32B32A32_TYPELESS: return DXGI_FORMAT_R32G32B32A32_UINT;
-	case DXGI_FORMAT_R32G32B32_TYPELESS:    return DXGI_FORMAT_R32G32B32_UINT;
-	case DXGI_FORMAT_R16G16B16A16_TYPELESS: return DXGI_FORMAT_R16G16B16A16_UNORM;
-	case DXGI_FORMAT_R32G32_TYPELESS:       return DXGI_FORMAT_R32G32_UINT;
-	case DXGI_FORMAT_R10G10B10A2_TYPELESS:  return DXGI_FORMAT_R10G10B10A2_UNORM;
-	case DXGI_FORMAT_R16G16_TYPELESS:       return DXGI_FORMAT_R16G16_UNORM;
-	case DXGI_FORMAT_R8G8_TYPELESS:         return DXGI_FORMAT_R8G8_UNORM;
-	case DXGI_FORMAT_R8_TYPELESS:           return DXGI_FORMAT_R8_UNORM;
-
-	case DXGI_FORMAT_BC4_TYPELESS:         return DXGI_FORMAT_BC4_UNORM;
-	case DXGI_FORMAT_BC5_TYPELESS:         return DXGI_FORMAT_BC5_UNORM;
-
-
-
-	case DXGI_FORMAT_R24G8_TYPELESS: return DXGI_FORMAT_R24_UNORM_X8_TYPELESS;
-	case DXGI_FORMAT_R32_TYPELESS: return DXGI_FORMAT_R32_FLOAT;
-	case DXGI_FORMAT_R16_TYPELESS: return DXGI_FORMAT_R16_UNORM;
-		// Changing Depth Buffers to 32 bit on Dingo as D24S8 is actually implemented as a 32 bit buffer in the hardware
-	case DXGI_FORMAT_R32G8X24_TYPELESS: return DXGI_FORMAT_R32_FLOAT_X8X24_TYPELESS;
-	}
-	return InFormat;
-}
-
-inline DXGI_FORMAT FindDepthStencilResourceDXGIFormat(DXGI_FORMAT InFormat)
-{
-	switch (InFormat)
-	{
-	case DXGI_FORMAT_R32_FLOAT: return DXGI_FORMAT_R32_TYPELESS;
-	case DXGI_FORMAT_R16_FLOAT: return DXGI_FORMAT_R16_TYPELESS;
-	}
-
-	return InFormat;
-}
-
-
-DXGI_FORMAT FD3D11DynamicRHI::GetPlatformTextureResourceFormat(DXGI_FORMAT InFormat, ETextureCreateFlags InFlags)
-{
-	// Find valid shared texture format
-	if (EnumHasAnyFlags(InFlags, TexCreate_Shared))
-	{
-		return FindSharedResourceDXGIFormat(InFormat, EnumHasAnyFlags(InFlags, TexCreate_SRGB));
-	}
-	if (EnumHasAnyFlags(InFlags, TexCreate_DepthStencilTargetable))
-	{
-		return FindDepthStencilResourceDXGIFormat(InFormat);
-	}
-	return InFormat;
-}
-
 // Work around an issue with the WARP device & BC7
 // Creating two views with different formats (DXGI_FORMAT_BC7_UNORM vs DXGI_FORMAT_BC7_UNORM_SRGB)
 // will result in a crash inside d3d10warp.dll when creating the second view
@@ -334,7 +254,7 @@ static void SafeCreateTexture2D(ID3D11Device* Direct3DDevice, int32 UEFormat, co
 				TextureDesc->Width,
 				TextureDesc->Height,
 				TextureDesc->ArraySize,
-				GetD3D11TextureFormatString(TextureDesc->Format),
+				UE::DXGIUtilities::GetFormatString(TextureDesc->Format),
 				(uint32)TextureDesc->Format,
 				TextureDesc->MipLevels,
 				UEFormat
@@ -380,9 +300,9 @@ FD3D11Texture* FD3D11DynamicRHI::CreateD3D11Texture2D(FRHITextureCreateDesc cons
 
 	const bool bSRGB = EnumHasAnyFlags(Flags, TexCreate_SRGB);
 
-	const DXGI_FORMAT PlatformResourceFormat = FD3D11DynamicRHI::GetPlatformTextureResourceFormat((DXGI_FORMAT)GPixelFormats[Format].PlatformFormat, Flags);
-	const DXGI_FORMAT PlatformShaderResourceFormat = FindShaderResourceDXGIFormat(PlatformResourceFormat, bSRGB);
-	const DXGI_FORMAT PlatformRenderTargetFormat = FindShaderResourceDXGIFormat(PlatformResourceFormat, bSRGB);
+	const DXGI_FORMAT PlatformResourceFormat = UE::DXGIUtilities::GetPlatformTextureResourceFormat((DXGI_FORMAT)GPixelFormats[Format].PlatformFormat, Flags);
+	const DXGI_FORMAT PlatformShaderResourceFormat = UE::DXGIUtilities::FindShaderResourceFormat(PlatformResourceFormat, bSRGB);
+	const DXGI_FORMAT PlatformRenderTargetFormat = UE::DXGIUtilities::FindShaderResourceFormat(PlatformResourceFormat, bSRGB);
 	
 	uint32 CPUAccessFlags = 0;
 	D3D11_USAGE TextureUsage = D3D11_USAGE_DEFAULT;
@@ -653,7 +573,7 @@ FD3D11Texture* FD3D11DynamicRHI::CreateD3D11Texture2D(FRHITextureCreateDesc cons
 		D3D11_DEPTH_STENCIL_VIEW_DESC DSVDesc;
 		FMemory::Memzero(DSVDesc);
 
-		DSVDesc.Format = FindDepthStencilDXGIFormat(PlatformResourceFormat);
+		DSVDesc.Format = UE::DXGIUtilities::FindDepthStencilFormat(PlatformResourceFormat);
 
 		if (bTextureArray || bCubeTexture)
 		{
@@ -691,7 +611,7 @@ FD3D11Texture* FD3D11DynamicRHI::CreateD3D11Texture2D(FRHITextureCreateDesc cons
 			// Read-only DSVs are not supported in Feature Level 10 so 
 			// a dummy DSV is created in order reduce logic complexity at a higher-level.
 			DSVDesc.Flags = (AccessType & FExclusiveDepthStencil::DepthRead_StencilWrite) ? D3D11_DSV_READ_ONLY_DEPTH : 0;
-			if(HasStencilBits(DSVDesc.Format))
+			if(UE::DXGIUtilities::HasStencilBits(DSVDesc.Format))
 			{
 				DSVDesc.Flags |= (AccessType & FExclusiveDepthStencil::DepthWrite_StencilRead) ? D3D11_DSV_READ_ONLY_STENCIL : 0;
 			}
@@ -794,8 +714,8 @@ FD3D11Texture* FD3D11DynamicRHI::CreateD3D11Texture3D(FRHITextureCreateDesc cons
 	const bool bSRGB = EnumHasAnyFlags(Flags, TexCreate_SRGB);
 
 	const DXGI_FORMAT PlatformResourceFormat = (DXGI_FORMAT)GPixelFormats[Format].PlatformFormat;
-	const DXGI_FORMAT PlatformShaderResourceFormat = FindShaderResourceDXGIFormat(PlatformResourceFormat, bSRGB);
-	const DXGI_FORMAT PlatformRenderTargetFormat = FindShaderResourceDXGIFormat(PlatformResourceFormat, bSRGB);
+	const DXGI_FORMAT PlatformShaderResourceFormat = UE::DXGIUtilities::FindShaderResourceFormat(PlatformResourceFormat, bSRGB);
+	const DXGI_FORMAT PlatformRenderTargetFormat = UE::DXGIUtilities::FindShaderResourceFormat(PlatformResourceFormat, bSRGB);
 
 	// Describe the texture.
 	D3D11_TEXTURE3D_DESC TextureDesc;
@@ -1556,9 +1476,9 @@ FD3D11Texture* FD3D11DynamicRHI::CreateTextureFromResource(bool bTextureArray, b
 
 	const bool bSRGB = EnumHasAnyFlags(TexCreateFlags, TexCreate_SRGB);
 
-	const DXGI_FORMAT PlatformResourceFormat = FD3D11DynamicRHI::GetPlatformTextureResourceFormat((DXGI_FORMAT)GPixelFormats[Format].PlatformFormat, TexCreateFlags);
-	const DXGI_FORMAT PlatformShaderResourceFormat = FindShaderResourceDXGIFormat(PlatformResourceFormat, bSRGB);
-	const DXGI_FORMAT PlatformRenderTargetFormat = FindShaderResourceDXGIFormat(PlatformResourceFormat, bSRGB);
+	const DXGI_FORMAT PlatformResourceFormat = UE::DXGIUtilities::GetPlatformTextureResourceFormat((DXGI_FORMAT)GPixelFormats[Format].PlatformFormat, TexCreateFlags);
+	const DXGI_FORMAT PlatformShaderResourceFormat = UE::DXGIUtilities::FindShaderResourceFormat(PlatformResourceFormat, bSRGB);
+	const DXGI_FORMAT PlatformRenderTargetFormat = UE::DXGIUtilities::FindShaderResourceFormat(PlatformResourceFormat, bSRGB);
 
 	const bool bIsMultisampled = TextureDesc.SampleDesc.Count > 1;
 
@@ -1665,7 +1585,7 @@ FD3D11Texture* FD3D11DynamicRHI::CreateTextureFromResource(bool bTextureArray, b
 		D3D11_DEPTH_STENCIL_VIEW_DESC DSVDesc;
 		FMemory::Memzero(DSVDesc);
 
-		DSVDesc.Format = FindDepthStencilDXGIFormat(PlatformResourceFormat);
+		DSVDesc.Format = UE::DXGIUtilities::FindDepthStencilFormat(PlatformResourceFormat);
 
 		if (bTextureArray || bCubeTexture)
 		{
@@ -1703,7 +1623,7 @@ FD3D11Texture* FD3D11DynamicRHI::CreateTextureFromResource(bool bTextureArray, b
 			// Read-only DSVs are not supported in Feature Level 10 so 
 			// a dummy DSV is created in order reduce logic complexity at a higher-level.
 			DSVDesc.Flags = (AccessType & FExclusiveDepthStencil::DepthRead_StencilWrite) ? D3D11_DSV_READ_ONLY_DEPTH : 0;
-			if (HasStencilBits(DSVDesc.Format))
+			if (UE::DXGIUtilities::HasStencilBits(DSVDesc.Format))
 			{
 				DSVDesc.Flags |= (AccessType & FExclusiveDepthStencil::DepthWrite_StencilRead) ? D3D11_DSV_READ_ONLY_STENCIL : 0;
 			}
diff --git a/Engine/Source/Runtime/Windows/D3D11RHI/Private/D3D11UAV.cpp b/Engine/Source/Runtime/Windows/D3D11RHI/Private/D3D11UAV.cpp
--- a/Engine/Source/Runtime/Windows/D3D11RHI/Private/D3D11UAV.cpp
+++ b/Engine/Source/Runtime/Windows/D3D11RHI/Private/D3D11UAV.cpp
@@ -33,7 +33,7 @@ void FD3D11UnorderedAccessView::UpdateView()
 		{
 			D3DResource = Buffer->Resource;
 
-			UAVDesc.Format = FindUnorderedAccessDXGIFormat(DXGI_FORMAT(GPixelFormats[Info.Format].PlatformFormat));
+			UAVDesc.Format = UE::DXGIUtilities::FindUnorderedAccessFormat(DXGI_FORMAT(GPixelFormats[Info.Format].PlatformFormat));
 			UAVDesc.ViewDimension = D3D11_UAV_DIMENSION_BUFFER;
 
 			UAVDesc.Buffer.FirstElement = Info.OffsetInBytes / Info.StrideInBytes;
@@ -67,7 +67,7 @@ void FD3D11UnorderedAccessView::UpdateView()
 		FRHITextureDesc const& TextureDesc = Texture->GetDesc();
 		auto const Info = ViewDesc.Texture.UAV.GetViewInfo(Texture);
 
-		UAVDesc.Format = FindUnorderedAccessDXGIFormat(DXGI_FORMAT(GPixelFormats[Info.Format].PlatformFormat));
+		UAVDesc.Format = UE::DXGIUtilities::FindUnorderedAccessFormat(DXGI_FORMAT(GPixelFormats[Info.Format].PlatformFormat));
 
 		// No need to use Info.Dimension, since D3D supports mixing Texture2D view types.
 		// Create a view which matches the underlying resource dimension.
diff --git a/Engine/Source/Runtime/Windows/D3D11RHI/Private/D3D11Util.cpp b/Engine/Source/Runtime/Windows/D3D11RHI/Private/D3D11Util.cpp
--- a/Engine/Source/Runtime/Windows/D3D11RHI/Private/D3D11Util.cpp
+++ b/Engine/Source/Runtime/Windows/D3D11RHI/Private/D3D11Util.cpp
@@ -60,50 +60,6 @@ FString GetD3D11ErrorString(HRESULT ErrorCode, ID3D11Device* Device)
 
 #undef D3DERR
 
-const TCHAR* GetD3D11TextureFormatString(DXGI_FORMAT TextureFormat)
-{
-	static const TCHAR* EmptyString = TEXT("");
-	const TCHAR* TextureFormatText = EmptyString;
-#define D3DFORMATCASE(x) case x: TextureFormatText = TEXT(#x); break;
-	switch(TextureFormat)
-	{
-		D3DFORMATCASE(DXGI_FORMAT_R8G8B8A8_UNORM)
-		D3DFORMATCASE(DXGI_FORMAT_B8G8R8A8_UNORM)
-		D3DFORMATCASE(DXGI_FORMAT_B8G8R8X8_UNORM)
-		D3DFORMATCASE(DXGI_FORMAT_BC1_UNORM)
-		D3DFORMATCASE(DXGI_FORMAT_BC2_UNORM)
-		D3DFORMATCASE(DXGI_FORMAT_BC3_UNORM)
-		D3DFORMATCASE(DXGI_FORMAT_BC4_UNORM)
-		D3DFORMATCASE(DXGI_FORMAT_R16G16B16A16_FLOAT)
-		D3DFORMATCASE(DXGI_FORMAT_R32G32B32A32_FLOAT)
-		D3DFORMATCASE(DXGI_FORMAT_UNKNOWN)
-		D3DFORMATCASE(DXGI_FORMAT_R8_UNORM)
-		D3DFORMATCASE(DXGI_FORMAT_D32_FLOAT_S8X24_UINT)
-		D3DFORMATCASE(DXGI_FORMAT_R32_FLOAT_X8X24_TYPELESS)
-		D3DFORMATCASE(DXGI_FORMAT_R32G8X24_TYPELESS)
-		D3DFORMATCASE(DXGI_FORMAT_D24_UNORM_S8_UINT)
-		D3DFORMATCASE(DXGI_FORMAT_R24_UNORM_X8_TYPELESS)
-		D3DFORMATCASE(DXGI_FORMAT_R32_FLOAT)
-		D3DFORMATCASE(DXGI_FORMAT_R16G16_UINT)
-		D3DFORMATCASE(DXGI_FORMAT_R16G16_UNORM)
-		D3DFORMATCASE(DXGI_FORMAT_R16G16_SNORM)
-		D3DFORMATCASE(DXGI_FORMAT_R16G16_FLOAT)
-		D3DFORMATCASE(DXGI_FORMAT_R32G32_FLOAT)
-		D3DFORMATCASE(DXGI_FORMAT_R10G10B10A2_UNORM)
-		D3DFORMATCASE(DXGI_FORMAT_R16G16B16A16_UINT)
-		D3DFORMATCASE(DXGI_FORMAT_R8G8_SNORM)
-		D3DFORMATCASE(DXGI_FORMAT_BC5_UNORM)
-		D3DFORMATCASE(DXGI_FORMAT_R1_UNORM)
-		D3DFORMATCASE(DXGI_FORMAT_R8G8B8A8_TYPELESS)
-		D3DFORMATCASE(DXGI_FORMAT_B8G8R8A8_TYPELESS)
-		D3DFORMATCASE(DXGI_FORMAT_BC7_UNORM)
-		D3DFORMATCASE(DXGI_FORMAT_BC6H_UF16)
-		default: TextureFormatText = EmptyString;
-	}
-#undef D3DFORMATCASE
-	return TextureFormatText;
-}
-
 static FString GetD3D11TextureFlagString(uint32 TextureFlags)
 {
 	FString TextureFormatText = TEXT("");
@@ -281,7 +237,7 @@ void VerifyD3D11CreateTextureResult(HRESULT D3DResult, int32 UEFormat,const ANSI
 	check(FAILED(D3DResult));
 
 	const FString ErrorString = GetD3D11ErrorString(D3DResult, 0);
-	const TCHAR* D3DFormatString = GetD3D11TextureFormatString((DXGI_FORMAT)D3DFormat);
+	const TCHAR* D3DFormatString = UE::DXGIUtilities::GetFormatString((DXGI_FORMAT)D3DFormat);
 
 	FString DebugInfoString;
 
@@ -354,8 +310,8 @@ void VerifyD3D11ResizeViewportResult(
 	check(FAILED(D3DResult));
 
 	const FString ErrorString = GetD3D11ErrorString(D3DResult, 0);
-	const TCHAR* OldStateFormat = GetD3D11TextureFormatString(OldState.Format);
-	const TCHAR* NewStateFormat = GetD3D11TextureFormatString(NewState.Format);
+	const TCHAR* OldStateFormat = UE::DXGIUtilities::GetFormatString(OldState.Format);
+	const TCHAR* NewStateFormat = UE::DXGIUtilities::GetFormatString(NewState.Format);
 
 	UE_LOG(LogD3D11RHI, Error,
 		TEXT("%s failed with error %s\n at %s:%u\n (Size=%ix%i Fullscreen=%d Format=%s(0x%08X)) -> (Size=%ix%i Fullscreen=%d Format=%s(0x%08X))"),
@@ -407,7 +363,7 @@ void VerifyD3D11CreateViewResult(HRESULT D3DResult, const ANSICHAR* Code, const
 	FormatSupport2.InFormat = Desc.Format;
 	Device->CheckFeatureSupport(D3D11_FEATURE_FORMAT_SUPPORT2, &FormatSupport2, sizeof(FormatSupport2));
 
-	const TCHAR* ViewFormat = GetD3D11TextureFormatString(Desc.Format);
+	const TCHAR* ViewFormat = UE::DXGIUtilities::GetFormatString(Desc.Format);
 
 	const FString& ErrorString = GetD3D11ErrorString(D3DResult, Device);
 
diff --git a/Engine/Source/Runtime/Windows/D3D11RHI/Private/Windows/WindowsD3D11Viewport.cpp b/Engine/Source/Runtime/Windows/D3D11RHI/Private/Windows/WindowsD3D11Viewport.cpp
--- a/Engine/Source/Runtime/Windows/D3D11RHI/Private/Windows/WindowsD3D11Viewport.cpp
+++ b/Engine/Source/Runtime/Windows/D3D11RHI/Private/Windows/WindowsD3D11Viewport.cpp
@@ -269,7 +269,7 @@ FD3D11Viewport::FD3D11Viewport(FD3D11DynamicRHI* InD3DRHI,HWND InWindowHandle,ui
 				D3DRHI->GetFactory()->CreateSwapChain(D3DRHI->GetDevice(), &SwapChainDesc, SwapChain.GetInitReference());
 			if (CreateSwapChainResult == E_INVALIDARG)
 			{
-				const TCHAR* D3DFormatString = GetD3D11TextureFormatString(SwapChainDesc.BufferDesc.Format);
+				const TCHAR* D3DFormatString = UE::DXGIUtilities::GetFormatString(SwapChainDesc.BufferDesc.Format);
 
 				UE_LOG(LogD3D11RHI, Error,
 					TEXT("CreateSwapChain failed with E_INVALIDARG: \n")
diff --git a/Engine/Source/Runtime/Windows/D3D11RHI/Public/D3D11Util.h b/Engine/Source/Runtime/Windows/D3D11RHI/Public/D3D11Util.h
--- a/Engine/Source/Runtime/Windows/D3D11RHI/Public/D3D11Util.h
+++ b/Engine/Source/Runtime/Windows/D3D11RHI/Public/D3D11Util.h
@@ -86,9 +86,6 @@ extern D3D11RHI_API void VerifyComRefCount(IUnknown* Object,int32 ExpectedRefs,c
 /** Returns a string for the provided error code, can include device removed information if the device is provided. */
 FString GetD3D11ErrorString(HRESULT ErrorCode, ID3D11Device* Device);
 
-/** Returns a string for the provided DXGI format. */
-const TCHAR* GetD3D11TextureFormatString(DXGI_FORMAT TextureFormat);
-
 /**
 * Convert from ECubeFace to D3DCUBEMAP_FACES type
 * @param Face - ECubeFace type to convert
