[Engine]

diff --git a/Engine/Source/Runtime/PhysicsCore/Public/CollisionQueryFilterCallbackCore.h b/Engine/Source/Runtime/PhysicsCore/Public/CollisionQueryFilterCallbackCore.h
--- a/Engine/Source/Runtime/PhysicsCore/Public/CollisionQueryFilterCallbackCore.h
+++ b/Engine/Source/Runtime/PhysicsCore/Public/CollisionQueryFilterCallbackCore.h
@@ -33,7 +33,7 @@ public:
 	virtual ECollisionQueryHitType PreFilter(const FCollisionFilterData& FilterData, const Chaos::FPerShapeData& Shape, const Chaos::FGeometryParticleHandle& Actor) = 0;
 };
 
-class PHYSICSCORE_API FBlockAllQueryCallback : public ICollisionQueryFilterCallbackBase
+class FBlockAllQueryCallback : public ICollisionQueryFilterCallbackBase
 {
 public:
 	virtual ~FBlockAllQueryCallback() {}
@@ -44,7 +44,7 @@ public:
 	virtual ECollisionQueryHitType PreFilter(const FCollisionFilterData& FilterData, const Chaos::FPerShapeData& Shape, const Chaos::FGeometryParticleHandle& Actor) override { return ECollisionQueryHitType::Block; }
 };
 
-class PHYSICSCORE_API FOverlapAllQueryCallback : public ICollisionQueryFilterCallbackBase
+class FOverlapAllQueryCallback : public ICollisionQueryFilterCallbackBase
 {
 public:
 	virtual ~FOverlapAllQueryCallback() {}
@@ -53,4 +53,4 @@ public:
 
 	virtual ECollisionQueryHitType PostFilter(const FCollisionFilterData& FilterData, const ChaosInterface::FPTQueryHit& Hit) override { return ECollisionQueryHitType::Touch; }
 	virtual ECollisionQueryHitType PreFilter(const FCollisionFilterData& FilterData, const Chaos::FPerShapeData& Shape, const Chaos::FGeometryParticleHandle& Actor) override { return ECollisionQueryHitType::Touch; }
-};
\ No newline at end of file
+};
diff --git a/Engine/Source/Runtime/PhysicsCore/Public/PhysTestSerializer.h b/Engine/Source/Runtime/PhysicsCore/Public/PhysTestSerializer.h
--- a/Engine/Source/Runtime/PhysicsCore/Public/PhysTestSerializer.h
+++ b/Engine/Source/Runtime/PhysicsCore/Public/PhysTestSerializer.h
@@ -33,20 +33,20 @@ namespace Chaos
 	class FChaosArchive;
 }
 
-class PHYSICSCORE_API FPhysTestSerializer
+class FPhysTestSerializer
 {
 public:
-	FPhysTestSerializer();
+	PHYSICSCORE_API FPhysTestSerializer();
 	FPhysTestSerializer(const FPhysTestSerializer& Other) = delete;
 	FPhysTestSerializer(FPhysTestSerializer&& Other) = delete;
 	FPhysTestSerializer& operator=(const FPhysTestSerializer&) = delete;
 	FPhysTestSerializer& operator=(FPhysTestSerializer&&) = delete;
 
-	void Serialize(Chaos::FChaosArchive& Ar);
-	void Serialize(const TCHAR* FilePrefix);
+	PHYSICSCORE_API void Serialize(Chaos::FChaosArchive& Ar);
+	PHYSICSCORE_API void Serialize(const TCHAR* FilePrefix);
 
 	//Set the data from an external source. This will obliterate any existing data. Make sure you are not holding on to old internal data as it will go away
-	void SetPhysicsData(Chaos::FPBDRigidsEvolution& ChaosEvolution);
+	PHYSICSCORE_API void SetPhysicsData(Chaos::FPBDRigidsEvolution& ChaosEvolution);
 
 	const Chaos::FChaosArchiveContext* GetChaosContext() const
 	{
@@ -89,7 +89,7 @@ public:
 private:
 
 #if 0
-	void CreateChaosData();
+	PHYSICSCORE_API void CreateChaosData();
 #endif 
 
 private:
diff --git a/Engine/Source/Runtime/PhysicsCore/Public/PhysicalMaterials/PhysicalMaterial.h b/Engine/Source/Runtime/PhysicsCore/Public/PhysicalMaterials/PhysicalMaterial.h
--- a/Engine/Source/Runtime/PhysicsCore/Public/PhysicalMaterials/PhysicalMaterial.h
+++ b/Engine/Source/Runtime/PhysicsCore/Public/PhysicalMaterials/PhysicalMaterial.h
@@ -55,13 +55,13 @@ struct FPhysicalMaterialStrength
 /**
  * Physical materials are used to define the response of a physical object when interacting dynamically with the world.
  */
-UCLASS(BlueprintType, Blueprintable, CollapseCategories, HideCategories = Object)
-class PHYSICSCORE_API UPhysicalMaterial : public UObject
+UCLASS(BlueprintType, Blueprintable, CollapseCategories, HideCategories = Object, MinimalAPI)
+class UPhysicalMaterial : public UObject
 {
 	GENERATED_UCLASS_BODY()
 
-	virtual ~UPhysicalMaterial();
-	UPhysicalMaterial(FVTableHelper& Helper);
+	PHYSICSCORE_API virtual ~UPhysicalMaterial();
+	PHYSICSCORE_API UPhysicalMaterial(FVTableHelper& Helper);
 	//
 	// Surface properties.
 	//
@@ -147,22 +147,22 @@ public:
 
 	//~ Begin UObject Interface
 #if WITH_EDITOR
-	virtual void PostEditChangeProperty(FPropertyChangedEvent& PropertyChangedEvent) override;
-	static void RebuildPhysicalMaterials();
+	PHYSICSCORE_API virtual void PostEditChangeProperty(FPropertyChangedEvent& PropertyChangedEvent) override;
+	static PHYSICSCORE_API void RebuildPhysicalMaterials();
 #endif // WITH_EDITOR
-	virtual void PostLoad() override;
-	virtual void FinishDestroy() override;
+	PHYSICSCORE_API virtual void PostLoad() override;
+	PHYSICSCORE_API virtual void FinishDestroy() override;
 	//~ End UObject Interface
 
 	/** Get the physics-interface derived version of this material */
-	FPhysicsMaterialHandle& GetPhysicsMaterial();
+	PHYSICSCORE_API FPhysicsMaterialHandle& GetPhysicsMaterial();
 
-	static void SetEngineDefaultPhysMaterial(UPhysicalMaterial* Material);
+	static PHYSICSCORE_API void SetEngineDefaultPhysMaterial(UPhysicalMaterial* Material);
 
-	static void SetEngineDefaultDestructiblePhysMaterial(UPhysicalMaterial* Material);
+	static PHYSICSCORE_API void SetEngineDefaultDestructiblePhysMaterial(UPhysicalMaterial* Material);
 
 	/** Determine Material Type from input PhysicalMaterial **/
-	static EPhysicalSurface DetermineSurfaceType(UPhysicalMaterial const* PhysicalMaterial);
+	static PHYSICSCORE_API EPhysicalSurface DetermineSurfaceType(UPhysicalMaterial const* PhysicalMaterial);
 };
 
 
diff --git a/Engine/Source/Runtime/PhysicsCore/Public/PhysicsObjectPhysicsCoreInterface.h b/Engine/Source/Runtime/PhysicsCore/Public/PhysicsObjectPhysicsCoreInterface.h
--- a/Engine/Source/Runtime/PhysicsCore/Public/PhysicsObjectPhysicsCoreInterface.h
+++ b/Engine/Source/Runtime/PhysicsCore/Public/PhysicsObjectPhysicsCoreInterface.h
@@ -7,8 +7,8 @@
 
 class FChaosScene;
 
-class PHYSICSCORE_API PhysicsObjectPhysicsCoreInterface
+class PhysicsObjectPhysicsCoreInterface
 {
 public:
-	static FChaosScene* GetScene(TArrayView<const Chaos::FConstPhysicsObjectHandle> InObjects);
-};
\ No newline at end of file
+	static PHYSICSCORE_API FChaosScene* GetScene(TArrayView<const Chaos::FConstPhysicsObjectHandle> InObjects);
+};
diff --git a/Engine/Source/Runtime/PhysicsCore/Public/PhysicsPublicCore.h b/Engine/Source/Runtime/PhysicsCore/Public/PhysicsPublicCore.h
--- a/Engine/Source/Runtime/PhysicsCore/Public/PhysicsPublicCore.h
+++ b/Engine/Source/Runtime/PhysicsCore/Public/PhysicsPublicCore.h
@@ -7,9 +7,9 @@
 class UPhysicalMaterial;
 
 /** Set of delegates to allowing hooking different parts of the physics engine */
-class PHYSICSCORE_API FPhysicsDelegatesCore
+class FPhysicsDelegatesCore
 {
 public:
 	DECLARE_MULTICAST_DELEGATE_OneParam(FOnUpdatePhysXMaterial, UPhysicalMaterial*);
-	static FOnUpdatePhysXMaterial OnUpdatePhysXMaterial;
+	static PHYSICSCORE_API FOnUpdatePhysXMaterial OnUpdatePhysXMaterial;
 };
diff --git a/Engine/Source/Runtime/PhysicsCore/Public/PhysicsSettingsCore.h b/Engine/Source/Runtime/PhysicsCore/Public/PhysicsSettingsCore.h
--- a/Engine/Source/Runtime/PhysicsCore/Public/PhysicsSettingsCore.h
+++ b/Engine/Source/Runtime/PhysicsCore/Public/PhysicsSettingsCore.h
@@ -15,8 +15,8 @@
 /**
  * Default physics settings.
  */
-UCLASS(config=Engine,defaultconfig,meta=(DisplayName="Physics"))
-class PHYSICSCORE_API UPhysicsSettingsCore: public UDeveloperSettings
+UCLASS(config=Engine,defaultconfig,meta=(DisplayName="Physics"), MinimalAPI)
+class UPhysicsSettingsCore: public UDeveloperSettings
 {
 	GENERATED_UCLASS_BODY()
 
@@ -122,15 +122,15 @@ class PHYSICSCORE_API UPhysicsSettingsCore: public UDeveloperSettings
 	UPROPERTY(config, EditAnywhere, Category = ChaosPhysics)
 	FChaosSolverConfiguration SolverOptions;
 
-	static UPhysicsSettingsCore* Get();
+	static PHYSICSCORE_API UPhysicsSettingsCore* Get();
 
-	virtual void PostInitProperties() override;
+	PHYSICSCORE_API virtual void PostInitProperties() override;
 
 protected:
-	static void SetDefaultSettings(UPhysicsSettingsCore* InSettings);
+	static PHYSICSCORE_API void SetDefaultSettings(UPhysicsSettingsCore* InSettings);
 
 private:
 	// Override default settings.
 	// This should be set up to point to the CDO of the leaf settings class (as edited in Project Settings)
 	static UPhysicsSettingsCore* DefaultSettings;
-};
\ No newline at end of file
+};
diff --git a/Engine/Source/Runtime/PhysicsCore/Public/SQAccelerator.h b/Engine/Source/Runtime/PhysicsCore/Public/SQAccelerator.h
--- a/Engine/Source/Runtime/PhysicsCore/Public/SQAccelerator.h
+++ b/Engine/Source/Runtime/PhysicsCore/Public/SQAccelerator.h
@@ -34,21 +34,21 @@ struct FCollisionFilterData;
 struct FCollisionQueryParams;
 struct FCollisionQueryParams;
 
-class PHYSICSCORE_API FChaosSQAccelerator
+class FChaosSQAccelerator
 {
 public:
 
-	FChaosSQAccelerator(const Chaos::ISpatialAcceleration<Chaos::FAccelerationStructureHandle,Chaos::FReal, 3>& InSpatialAcceleration);
+	PHYSICSCORE_API FChaosSQAccelerator(const Chaos::ISpatialAcceleration<Chaos::FAccelerationStructureHandle,Chaos::FReal, 3>& InSpatialAcceleration);
 	virtual ~FChaosSQAccelerator() {};
 
-	void Raycast(const FVector& Start, const FVector& Dir, const float DeltaMagnitude, ChaosInterface::FSQHitBuffer<ChaosInterface::FRaycastHit>& HitBuffer, EHitFlags OutputFlags, const ChaosInterface::FQueryFilterData& QueryFilterData, ICollisionQueryFilterCallbackBase& QueryCallback, const ChaosInterface::FQueryDebugParams& DebugParams = {}) const;
-	void Raycast(const FVector& Start, const FVector& Dir, const float DeltaMagnitude, ChaosInterface::FSQHitBuffer<ChaosInterface::FPTRaycastHit>& HitBuffer, EHitFlags OutputFlags, const ChaosInterface::FQueryFilterData& QueryFilterData, ICollisionQueryFilterCallbackBase& QueryCallback, const ChaosInterface::FQueryDebugParams& DebugParams = {}) const;
+	PHYSICSCORE_API void Raycast(const FVector& Start, const FVector& Dir, const float DeltaMagnitude, ChaosInterface::FSQHitBuffer<ChaosInterface::FRaycastHit>& HitBuffer, EHitFlags OutputFlags, const ChaosInterface::FQueryFilterData& QueryFilterData, ICollisionQueryFilterCallbackBase& QueryCallback, const ChaosInterface::FQueryDebugParams& DebugParams = {}) const;
+	PHYSICSCORE_API void Raycast(const FVector& Start, const FVector& Dir, const float DeltaMagnitude, ChaosInterface::FSQHitBuffer<ChaosInterface::FPTRaycastHit>& HitBuffer, EHitFlags OutputFlags, const ChaosInterface::FQueryFilterData& QueryFilterData, ICollisionQueryFilterCallbackBase& QueryCallback, const ChaosInterface::FQueryDebugParams& DebugParams = {}) const;
 
-	void Sweep(const Chaos::FImplicitObject& QueryGeom, const FTransform& StartTM, const FVector& Dir, const float DeltaMagnitude, ChaosInterface::FSQHitBuffer<ChaosInterface::FSweepHit>& HitBuffer, EHitFlags OutputFlags, const ChaosInterface::FQueryFilterData& QueryFilterData, ICollisionQueryFilterCallbackBase& QueryCallback, const ChaosInterface::FQueryDebugParams& DebugParams = {}) const;
-	void Sweep(const Chaos::FImplicitObject& QueryGeom, const FTransform& StartTM, const FVector& Dir, const float DeltaMagnitude, ChaosInterface::FSQHitBuffer<ChaosInterface::FPTSweepHit>& HitBuffer, EHitFlags OutputFlags, const ChaosInterface::FQueryFilterData& QueryFilterData, ICollisionQueryFilterCallbackBase& QueryCallback, const ChaosInterface::FQueryDebugParams& DebugParams = {}) const;
+	PHYSICSCORE_API void Sweep(const Chaos::FImplicitObject& QueryGeom, const FTransform& StartTM, const FVector& Dir, const float DeltaMagnitude, ChaosInterface::FSQHitBuffer<ChaosInterface::FSweepHit>& HitBuffer, EHitFlags OutputFlags, const ChaosInterface::FQueryFilterData& QueryFilterData, ICollisionQueryFilterCallbackBase& QueryCallback, const ChaosInterface::FQueryDebugParams& DebugParams = {}) const;
+	PHYSICSCORE_API void Sweep(const Chaos::FImplicitObject& QueryGeom, const FTransform& StartTM, const FVector& Dir, const float DeltaMagnitude, ChaosInterface::FSQHitBuffer<ChaosInterface::FPTSweepHit>& HitBuffer, EHitFlags OutputFlags, const ChaosInterface::FQueryFilterData& QueryFilterData, ICollisionQueryFilterCallbackBase& QueryCallback, const ChaosInterface::FQueryDebugParams& DebugParams = {}) const;
 
-	void Overlap(const Chaos::FImplicitObject& QueryGeom, const FTransform& GeomPose, ChaosInterface::FSQHitBuffer<ChaosInterface::FOverlapHit>& HitBuffer, const ChaosInterface::FQueryFilterData& QueryFilterData, ICollisionQueryFilterCallbackBase& QueryCallback, const ChaosInterface::FQueryDebugParams& DebugParams = {}) const;
-	void Overlap(const Chaos::FImplicitObject& QueryGeom, const FTransform& GeomPose, ChaosInterface::FSQHitBuffer<ChaosInterface::FPTOverlapHit>& HitBuffer, const ChaosInterface::FQueryFilterData& QueryFilterData, ICollisionQueryFilterCallbackBase& QueryCallback, const ChaosInterface::FQueryDebugParams& DebugParams = {}) const;
+	PHYSICSCORE_API void Overlap(const Chaos::FImplicitObject& QueryGeom, const FTransform& GeomPose, ChaosInterface::FSQHitBuffer<ChaosInterface::FOverlapHit>& HitBuffer, const ChaosInterface::FQueryFilterData& QueryFilterData, ICollisionQueryFilterCallbackBase& QueryCallback, const ChaosInterface::FQueryDebugParams& DebugParams = {}) const;
+	PHYSICSCORE_API void Overlap(const Chaos::FImplicitObject& QueryGeom, const FTransform& GeomPose, ChaosInterface::FSQHitBuffer<ChaosInterface::FPTOverlapHit>& HitBuffer, const ChaosInterface::FQueryFilterData& QueryFilterData, ICollisionQueryFilterCallbackBase& QueryCallback, const ChaosInterface::FQueryDebugParams& DebugParams = {}) const;
 
 private:
 	const Chaos::ISpatialAcceleration<Chaos::FAccelerationStructureHandle, Chaos::FReal, 3>& SpatialAcceleration;
@@ -62,7 +62,7 @@ private:
 // This was used in the 2019 GDC demos and is now broken. To make it work again, we would need to implement
 // the FChaosSQAcceleratorAdapter below to use its internal SQ accelerator and convert the inputs and outputs
 // from/to PhysX types.
-class PHYSICSCORE_API ISQAccelerator
+class ISQAccelerator
 {
 public:
 	virtual ~ISQAccelerator() {};
@@ -71,16 +71,16 @@ public:
 	virtual void Overlap(const FPhysicsGeometry& QueryGeom, const FTransform& GeomPose, FPhysicsHitCallback<FHitOverlap>& HitBuffer, const ChaosInterface::FQueryFilterData& QueryFilterData, ICollisionQueryFilterCallbackBase& QueryCallback) const = 0;
 };
 
-class PHYSICSCORE_API FSQAcceleratorUnion : public ISQAccelerator
+class FSQAcceleratorUnion : public ISQAccelerator
 {
 public:
 
-	virtual void Raycast(const FVector& Start, const FVector& Dir, const float DeltaMagnitude, FPhysicsHitCallback<FHitRaycast>& HitBuffer, EHitFlags OutputFlags, const ChaosInterface::FQueryFilterData& QueryFilterData, ICollisionQueryFilterCallbackBase& QueryCallback) const override;
-	virtual void Sweep(const FPhysicsGeometry& QueryGeom, const FTransform& StartTM, const FVector& Dir, const float DeltaMagnitude, FPhysicsHitCallback<FHitSweep>& HitBuffer, EHitFlags OutputFlags, const ChaosInterface::FQueryFilterData& QueryFilterData, ICollisionQueryFilterCallbackBase& QueryCallback) const override;
-	virtual void Overlap(const FPhysicsGeometry& QueryGeom, const FTransform& GeomPose, FPhysicsHitCallback<FHitOverlap>& HitBuffer, const ChaosInterface::FQueryFilterData& QueryFilterData, ICollisionQueryFilterCallbackBase& QueryCallback) const override;
+	PHYSICSCORE_API virtual void Raycast(const FVector& Start, const FVector& Dir, const float DeltaMagnitude, FPhysicsHitCallback<FHitRaycast>& HitBuffer, EHitFlags OutputFlags, const ChaosInterface::FQueryFilterData& QueryFilterData, ICollisionQueryFilterCallbackBase& QueryCallback) const override;
+	PHYSICSCORE_API virtual void Sweep(const FPhysicsGeometry& QueryGeom, const FTransform& StartTM, const FVector& Dir, const float DeltaMagnitude, FPhysicsHitCallback<FHitSweep>& HitBuffer, EHitFlags OutputFlags, const ChaosInterface::FQueryFilterData& QueryFilterData, ICollisionQueryFilterCallbackBase& QueryCallback) const override;
+	PHYSICSCORE_API virtual void Overlap(const FPhysicsGeometry& QueryGeom, const FTransform& GeomPose, FPhysicsHitCallback<FHitOverlap>& HitBuffer, const ChaosInterface::FQueryFilterData& QueryFilterData, ICollisionQueryFilterCallbackBase& QueryCallback) const override;
 
-	void AddSQAccelerator(ISQAccelerator* InAccelerator);
-	void RemoveSQAccelerator(ISQAccelerator* AcceleratorToRemove);
+	PHYSICSCORE_API void AddSQAccelerator(ISQAccelerator* InAccelerator);
+	PHYSICSCORE_API void RemoveSQAccelerator(ISQAccelerator* AcceleratorToRemove);
 
 private:
 	TArray<ISQAccelerator*> Accelerators;
diff --git a/Engine/Source/Runtime/PhysicsCore/Public/SQCapture.h b/Engine/Source/Runtime/PhysicsCore/Public/SQCapture.h
--- a/Engine/Source/Runtime/PhysicsCore/Public/SQCapture.h
+++ b/Engine/Source/Runtime/PhysicsCore/Public/SQCapture.h
@@ -37,9 +37,9 @@ namespace Chaos
 }
 
 //Allows us to capture a scene query with either physx or chaos and then load it into either format for testing purposes
-struct PHYSICSCORE_API FSQCapture
+struct FSQCapture
 {
-	~FSQCapture();
+	PHYSICSCORE_API ~FSQCapture();
 	FSQCapture(const FSQCapture&) = delete;
 	FSQCapture& operator=(const FSQCapture&) = delete;
 
@@ -50,16 +50,16 @@ struct PHYSICSCORE_API FSQCapture
 		Overlap
 	} SQType;
 
-	void StartCaptureChaosSweep(const Chaos::FPBDRigidsEvolution& Evolution, const Chaos::FImplicitObject& InQueryGeom, const FTransform& InStartTM, const FVector& InDir, float InDeltaMag, FHitFlags InOutputFlags, const ChaosInterface::FQueryFilterData& QueryFilter, const FCollisionFilterData& FilterData, ICollisionQueryFilterCallbackBase& Callback);
-	void EndCaptureChaosSweep(const ChaosInterface::FSQHitBuffer<ChaosInterface::FSweepHit>& Results);
+	PHYSICSCORE_API void StartCaptureChaosSweep(const Chaos::FPBDRigidsEvolution& Evolution, const Chaos::FImplicitObject& InQueryGeom, const FTransform& InStartTM, const FVector& InDir, float InDeltaMag, FHitFlags InOutputFlags, const ChaosInterface::FQueryFilterData& QueryFilter, const FCollisionFilterData& FilterData, ICollisionQueryFilterCallbackBase& Callback);
+	PHYSICSCORE_API void EndCaptureChaosSweep(const ChaosInterface::FSQHitBuffer<ChaosInterface::FSweepHit>& Results);
 
-	void StartCaptureChaosRaycast(const Chaos::FPBDRigidsEvolution& Evolution, const FVector& InStartPoint, const FVector& InDir, float InDeltaMag, FHitFlags InOutputFlags, const ChaosInterface::FQueryFilterData& QueryFilter, const FCollisionFilterData& FilterData, ICollisionQueryFilterCallbackBase& Callback);
-	void EndCaptureChaosRaycast(const ChaosInterface::FSQHitBuffer<ChaosInterface::FRaycastHit>& Results);
+	PHYSICSCORE_API void StartCaptureChaosRaycast(const Chaos::FPBDRigidsEvolution& Evolution, const FVector& InStartPoint, const FVector& InDir, float InDeltaMag, FHitFlags InOutputFlags, const ChaosInterface::FQueryFilterData& QueryFilter, const FCollisionFilterData& FilterData, ICollisionQueryFilterCallbackBase& Callback);
+	PHYSICSCORE_API void EndCaptureChaosRaycast(const ChaosInterface::FSQHitBuffer<ChaosInterface::FRaycastHit>& Results);
 
-	void StartCaptureChaosOverlap(const Chaos::FPBDRigidsEvolution& Evolution, const Chaos::FImplicitObject& InQueryGeom, const FTransform& InStartTM, const ChaosInterface::FQueryFilterData& QueryFilter, const FCollisionFilterData& FilterData, ICollisionQueryFilterCallbackBase& Callback);
-	void EndCaptureChaosOverlap(const ChaosInterface::FSQHitBuffer<ChaosInterface::FOverlapHit>& Results);
+	PHYSICSCORE_API void StartCaptureChaosOverlap(const Chaos::FPBDRigidsEvolution& Evolution, const Chaos::FImplicitObject& InQueryGeom, const FTransform& InStartTM, const ChaosInterface::FQueryFilterData& QueryFilter, const FCollisionFilterData& FilterData, ICollisionQueryFilterCallbackBase& Callback);
+	PHYSICSCORE_API void EndCaptureChaosOverlap(const ChaosInterface::FSQHitBuffer<ChaosInterface::FOverlapHit>& Results);
 
-	ECollisionQueryHitType GetFilterResult(const Chaos::FPerShapeData* Shape, const Chaos::FGeometryParticle* Actor) const;
+	PHYSICSCORE_API ECollisionQueryHitType GetFilterResult(const Chaos::FPerShapeData* Shape, const Chaos::FGeometryParticle* Actor) const;
 	
 	FVector Dir;
 	FTransform StartTM;	//only valid if overlap or sweep
@@ -84,8 +84,8 @@ struct PHYSICSCORE_API FSQCapture
 	TArray<ChaosInterface::FOverlapHit> ChaosOverlapTouches;
 
 private:
-	FSQCapture(FPhysTestSerializer& OwningPhysSerializer);	//This should be created by PhysTestSerializer
-	void Serialize(Chaos::FChaosArchive& Ar);
+	PHYSICSCORE_API FSQCapture(FPhysTestSerializer& OwningPhysSerializer);	//This should be created by PhysTestSerializer
+	PHYSICSCORE_API void Serialize(Chaos::FChaosArchive& Ar);
 
 	friend FPhysTestSerializer;
 
@@ -94,14 +94,14 @@ private:
 
 	FPhysTestSerializer& PhysSerializer;
 
-	void CaptureChaosFilterResults(const Chaos::FPBDRigidsEvolution& Evolution, const FCollisionFilterData& FilterData, ICollisionQueryFilterCallbackBase& Callback);
+	PHYSICSCORE_API void CaptureChaosFilterResults(const Chaos::FPBDRigidsEvolution& Evolution, const FCollisionFilterData& FilterData, ICollisionQueryFilterCallbackBase& Callback);
 
 	TMap<Chaos::FGeometryParticle*, TArray<TPair<Chaos::FPerShapeData*, ECollisionQueryHitType>>> ChaosActorToShapeHitsArray;
 
 	template <typename THit>
 	void SerializeChaosBuffers(Chaos::FChaosArchive& Ar, int32 Version, ChaosInterface::FSQHitBuffer<THit>& ChaosBuffer);
 
-	void SerializeChaosActorToShapeHitsArray(Chaos::FChaosArchive& Ar);
+	PHYSICSCORE_API void SerializeChaosActorToShapeHitsArray(Chaos::FChaosArchive& Ar);
 
 	bool bDiskDataIsChaos;
 	bool bChaosDataReady;
diff --git a/Engine/Source/Runtime/Portal/LauncherPlatform/Public/ILauncherPlatform.h b/Engine/Source/Runtime/Portal/LauncherPlatform/Public/ILauncherPlatform.h
--- a/Engine/Source/Runtime/Portal/LauncherPlatform/Public/ILauncherPlatform.h
+++ b/Engine/Source/Runtime/Portal/LauncherPlatform/Public/ILauncherPlatform.h
@@ -104,9 +104,9 @@ public:
 	virtual bool OpenLauncher(const FOpenLauncherOptions& Options) = 0;
 };
 
-class LAUNCHERPLATFORM_API FLauncherMisc
+class FLauncherMisc
 {
 public:
 	/** Return url encoded full path of currently running executable. */
-	static FString GetEncodedExePath();
-};
\ No newline at end of file
+	static LAUNCHERPLATFORM_API FString GetEncodedExePath();
+};
diff --git a/Engine/Source/Runtime/PreLoadScreen/Public/BuildPatchServicesPreLoadManager.h b/Engine/Source/Runtime/PreLoadScreen/Public/BuildPatchServicesPreLoadManager.h
--- a/Engine/Source/Runtime/PreLoadScreen/Public/BuildPatchServicesPreLoadManager.h
+++ b/Engine/Source/Runtime/PreLoadScreen/Public/BuildPatchServicesPreLoadManager.h
@@ -16,21 +16,21 @@
 class IBuildPatchServicesModule;
 
 //This class is used to help manage a PreLoadScreen based on a BuildPatchServices install.
-class PRELOADSCREEN_API FBuildPatchServicesPreLoadManagerBase : public TSharedFromThis<FBuildPatchServicesPreLoadManagerBase>
+class FBuildPatchServicesPreLoadManagerBase : public TSharedFromThis<FBuildPatchServicesPreLoadManagerBase>
 {
 public:
-    FBuildPatchServicesPreLoadManagerBase();
+    PRELOADSCREEN_API FBuildPatchServicesPreLoadManagerBase();
     virtual ~FBuildPatchServicesPreLoadManagerBase() {}
 
-    virtual void Init();
+    PRELOADSCREEN_API virtual void Init();
     
     //Setup BPT with everything now loaded
-    virtual void StartBuildPatchServices(BuildPatchServices::FBuildInstallerConfiguration Settings);
+    PRELOADSCREEN_API virtual void StartBuildPatchServices(BuildPatchServices::FBuildInstallerConfiguration Settings);
 
     //BPT finished
-    virtual void OnContentBuildInstallerComplete(const IBuildInstallerRef& Installer);
+    PRELOADSCREEN_API virtual void OnContentBuildInstallerComplete(const IBuildInstallerRef& Installer);
 
-    virtual bool IsDone() const;
+    PRELOADSCREEN_API virtual bool IsDone() const;
 
     virtual int64 GetDownloadSize() { return ContentBuildInstaller.IsValid() ? ContentBuildInstaller->GetTotalDownloadRequired() : 0; }
     virtual int64 GetDownloadProgress() { return ContentBuildInstaller.IsValid() ? ContentBuildInstaller->GetTotalDownloaded() : 0; }
@@ -38,13 +38,13 @@ public:
     DECLARE_MULTICAST_DELEGATE_OneParam(FOnBuildPatchCompleted, bool );
     FOnBuildPatchCompleted OnBuildPatchCompletedDelegate;
 
-    virtual void PauseBuildPatchInstall();
-    virtual void ResumeBuildPatchInstall();
-	virtual void CancelBuildPatchInstall();
+    PRELOADSCREEN_API virtual void PauseBuildPatchInstall();
+    PRELOADSCREEN_API virtual void ResumeBuildPatchInstall();
+	PRELOADSCREEN_API virtual void CancelBuildPatchInstall();
 
     float GetProgressPercent() const { return ContentBuildInstaller.IsValid() ? ContentBuildInstaller->GetUpdateProgress() : 0.0f; }
     EBuildPatchDownloadHealth GetDownloadHealth() const { return ContentBuildInstaller.IsValid() ? ContentBuildInstaller->GetDownloadHealth() : EBuildPatchDownloadHealth::NUM_Values; }
-    const FText& GetStatusText() const;
+    PRELOADSCREEN_API const FText& GetStatusText() const;
     BuildPatchServices::EBuildPatchState GetState() const { return ContentBuildInstaller.IsValid() ? ContentBuildInstaller->GetState() : BuildPatchServices::EBuildPatchState::Initializing; }
 
     FText GetErrorMessageBody() const { return ContentBuildInstaller.IsValid() ? ContentBuildInstaller->GetErrorText() : FText(); }
diff --git a/Engine/Source/Runtime/PreLoadScreen/Public/PreLoadScreenBase.h b/Engine/Source/Runtime/PreLoadScreen/Public/PreLoadScreenBase.h
--- a/Engine/Source/Runtime/PreLoadScreen/Public/PreLoadScreenBase.h
+++ b/Engine/Source/Runtime/PreLoadScreen/Public/PreLoadScreenBase.h
@@ -6,7 +6,7 @@
 
 // Base implementation of the IPreLoadScreen that handles all the logic for controlling / updating the UI for PreLoadScreens.
 // Designed to be overriden by a game specific Plugin that calls FPreloadScreenManager::RegisterPreLoadScreen so that functions are called by the PreLoadScreenManager correctly.
-class PRELOADSCREEN_API FPreLoadScreenBase : public IPreLoadScreen 
+class FPreLoadScreenBase : public IPreLoadScreen 
 {
 
     /**** IPreLoadScreen implementation ****/
@@ -37,10 +37,10 @@ public:
 	// PreLoadScreens not using this functionality should return NAME_None
 	virtual FName GetPreLoadScreenTag() const override { return NAME_None; }
 
-    virtual void CleanUp() override;
+    PRELOADSCREEN_API virtual void CleanUp() override;
 
     //Default behavior is just to see if we have an active widget. Should really overload with our own behavior to see if we are done displaying
-    virtual bool IsDone() const override;
+    PRELOADSCREEN_API virtual bool IsDone() const override;
 
 public:
     FPreLoadScreenBase()
@@ -50,7 +50,7 @@ public:
     virtual ~FPreLoadScreenBase() override {};
 
     //Handles constructing a FPreLoadSettingsContainerBase with the 
-    virtual void InitSettingsFromConfig(const FString& ConfigFileName);
+    PRELOADSCREEN_API virtual void InitSettingsFromConfig(const FString& ConfigFileName);
 
     //Set what plugin is creating this PreLoadScreenBase. Used to make file paths relative to that plugin as well as
     //determining . Used for converting locations for content to be relative to the plugin calling us
@@ -65,4 +65,4 @@ private:
     //The name of the Plugin creating this FPreLoadScreenBase.
     //Important: Should be set before Initting settings from Config!
     FString PluginName;
-};
\ No newline at end of file
+};
diff --git a/Engine/Source/Runtime/PreLoadScreen/Public/PreLoadScreenManager.h b/Engine/Source/Runtime/PreLoadScreen/Public/PreLoadScreenManager.h
--- a/Engine/Source/Runtime/PreLoadScreen/Public/PreLoadScreenManager.h
+++ b/Engine/Source/Runtime/PreLoadScreen/Public/PreLoadScreenManager.h
@@ -17,50 +17,50 @@
 
 
 // Class that handles storing all registered PreLoadScreens and Playing/Stopping them
-class PRELOADSCREEN_API FPreLoadScreenManager
+class FPreLoadScreenManager
 {
 public:
     //Gets the single instance of this settings object. Also creates it if needed
-    static FPreLoadScreenManager* Get();
-    static void Create();
-    static void Destroy();
+    static PRELOADSCREEN_API FPreLoadScreenManager* Get();
+    static PRELOADSCREEN_API void Create();
+    static PRELOADSCREEN_API void Destroy();
 
-    void Initialize(FSlateRenderer& InSlateRenderer);
+    PRELOADSCREEN_API void Initialize(FSlateRenderer& InSlateRenderer);
 
-    void RegisterPreLoadScreen(const TSharedPtr<IPreLoadScreen>& PreLoadScreen);
-    void UnRegisterPreLoadScreen(const TSharedPtr<IPreLoadScreen>& PreLoadScreen);
+    PRELOADSCREEN_API void RegisterPreLoadScreen(const TSharedPtr<IPreLoadScreen>& PreLoadScreen);
+    PRELOADSCREEN_API void UnRegisterPreLoadScreen(const TSharedPtr<IPreLoadScreen>& PreLoadScreen);
 
     /**
 	 * Plays the first found PreLoadScreen that matches the bEarlyPreLoadScreen setting passed in.
 	 * @returns false if no PreLoadScreen with that type has been registered.
 	 */
-    bool PlayFirstPreLoadScreen(EPreLoadScreenTypes PreLoadScreenTypeToPlay);
+    PRELOADSCREEN_API bool PlayFirstPreLoadScreen(EPreLoadScreenTypes PreLoadScreenTypeToPlay);
 
 	/**
 	 * Plays the PreLoadScreen with a tag that matches InTag
 	 * @returns false if no PreLoadScreen with that tag has been registered.
 	 */
-	bool PlayPreLoadScreenWithTag(FName InTag);
+	PRELOADSCREEN_API bool PlayPreLoadScreenWithTag(FName InTag);
     
-	void StopPreLoadScreen();
-	void WaitForEngineLoadingScreenToFinish();
+	PRELOADSCREEN_API void StopPreLoadScreen();
+	PRELOADSCREEN_API void WaitForEngineLoadingScreenToFinish();
 
-    void PassPreLoadScreenWindowBackToGame() const;
+    PRELOADSCREEN_API void PassPreLoadScreenWindowBackToGame() const;
 
-    bool IsUsingMainWindow() const;
+    PRELOADSCREEN_API bool IsUsingMainWindow() const;
 
-    TSharedPtr<SWindow> GetRenderWindow();
+    PRELOADSCREEN_API TSharedPtr<SWindow> GetRenderWindow();
 
-    bool HasRegisteredPreLoadScreenType(EPreLoadScreenTypes PreLoadScreenTypeToCheck) const;
-    bool HasActivePreLoadScreenType(EPreLoadScreenTypes PreLoadScreenTypeToCheck) const;
-    bool HasValidActivePreLoadScreen() const;
+    PRELOADSCREEN_API bool HasRegisteredPreLoadScreenType(EPreLoadScreenTypes PreLoadScreenTypeToCheck) const;
+    PRELOADSCREEN_API bool HasActivePreLoadScreenType(EPreLoadScreenTypes PreLoadScreenTypeToCheck) const;
+    PRELOADSCREEN_API bool HasValidActivePreLoadScreen() const;
 
 
-    void SetEngineLoadingComplete(bool IsEngineLoadingFinished = true);
+    PRELOADSCREEN_API void SetEngineLoadingComplete(bool IsEngineLoadingFinished = true);
     bool IsEngineLoadingComplete() const { return bIsEngineLoadingComplete; }
 
-    static void EnableRendering(bool bEnabled);
-    static bool ArePreLoadScreensEnabled();
+    static PRELOADSCREEN_API void EnableRendering(bool bEnabled);
+    static PRELOADSCREEN_API bool ArePreLoadScreensEnabled();
     
     // Callback for handling cleaning up any resources you would like to remove after the PreLoadScreenManager cleans up
     // Not needed for PreLoadScreens as those have a separate CleanUp method called.
@@ -72,46 +72,46 @@ public:
 
 protected:
     //Default constructor. We don't want other classes to make these. Should just rely on Get()
-    FPreLoadScreenManager();
+    PRELOADSCREEN_API FPreLoadScreenManager();
 
 	~FPreLoadScreenManager() = default;
 
-	void PlayPreLoadScreenAtIndex(int32 Index);
+	PRELOADSCREEN_API void PlayPreLoadScreenAtIndex(int32 Index);
 
-    void BeginPlay();
+    PRELOADSCREEN_API void BeginPlay();
 
     /*** These functions describe the flow for an EarlyPreLoadScreen where everything is blocking waiting on a call to StopPreLoadScreen ***/
-    void HandleEarlyStartupPlay();
+    PRELOADSCREEN_API void HandleEarlyStartupPlay();
     //Separate tick that handles 
-    void EarlyPlayFrameTick();
-    void GameLogicFrameTick();
-    void EarlyPlayRenderFrameTick();
-	bool HasActivePreLoadScreenTypeForEarlyStartup() const;
+    PRELOADSCREEN_API void EarlyPlayFrameTick();
+    PRELOADSCREEN_API void GameLogicFrameTick();
+    PRELOADSCREEN_API void EarlyPlayRenderFrameTick();
+	PRELOADSCREEN_API bool HasActivePreLoadScreenTypeForEarlyStartup() const;
 
-	void PlatformSpecificGameLogicFrameTick();
+	PRELOADSCREEN_API void PlatformSpecificGameLogicFrameTick();
 
 	//Creates a tick on the Render Thread that we run every
-	static void StaticRenderTick_RenderThread();
-	void RenderTick_RenderThread();
+	static PRELOADSCREEN_API void StaticRenderTick_RenderThread();
+	PRELOADSCREEN_API void RenderTick_RenderThread();
 
     /*** These functions describe how everything is handled during an non-Early PreLoadPlay. Everything is handled asynchronously in this case with a standalone renderer ***/
-	void HandleEngineLoadingPlay();
+	PRELOADSCREEN_API void HandleEngineLoadingPlay();
 
 	/*** These functions describe the flow for showing an CustomSplashScreen ***/
-	void HandleCustomSplashScreenPlay();
+	PRELOADSCREEN_API void HandleCustomSplashScreenPlay();
 
-    IPreLoadScreen* GetActivePreLoadScreen();
-    const IPreLoadScreen* GetActivePreLoadScreen() const;
+    PRELOADSCREEN_API IPreLoadScreen* GetActivePreLoadScreen();
+    PRELOADSCREEN_API const IPreLoadScreen* GetActivePreLoadScreen() const;
 
-	void HandleStopPreLoadScreen();
-	void CleanUpResources();
+	PRELOADSCREEN_API void HandleStopPreLoadScreen();
+	PRELOADSCREEN_API void CleanUpResources();
 
 	//Helpers that setup and clean-up delegates that only need to be active while we are playing an EarlyStartup PreLoadScreen
-	void RegisterDelegatesForEarlyStartupPlay();
-	void CleanUpDelegatesForEarlyStartupPlay();
+	PRELOADSCREEN_API void RegisterDelegatesForEarlyStartupPlay();
+	PRELOADSCREEN_API void CleanUpDelegatesForEarlyStartupPlay();
 
-	void HandleFlushRenderingCommandsStart();
-	void HandleFlushRenderingCommandsEnd();
+	PRELOADSCREEN_API void HandleFlushRenderingCommandsStart();
+	PRELOADSCREEN_API void HandleFlushRenderingCommandsEnd();
 
 	//Singleton Instance
 	struct FPreLoadScreenManagerDelete
@@ -121,10 +121,10 @@ protected:
 			delete Ptr;
 		}
 	};
-	static TUniquePtr<FPreLoadScreenManager, FPreLoadScreenManagerDelete> Instance;
+	static PRELOADSCREEN_API TUniquePtr<FPreLoadScreenManager, FPreLoadScreenManagerDelete> Instance;
 
-	static FCriticalSection AcquireCriticalSection;
-	static TAtomic<bool> bRenderingEnabled;
+	static PRELOADSCREEN_API FCriticalSection AcquireCriticalSection;
+	static PRELOADSCREEN_API TAtomic<bool> bRenderingEnabled;
 
 	TArray<TSharedPtr<IPreLoadScreen>> PreLoadScreens;
 
@@ -156,10 +156,10 @@ protected:
 
 private:
 #if PLATFORM_ANDROID
-	void Android_PlatformSpecificGameLogicFrameTick();
+	PRELOADSCREEN_API void Android_PlatformSpecificGameLogicFrameTick();
 #endif //PLATFORM_ANDROID
 
 #if PLATFORM_IOS
-	void IOS_PlatformSpecificGameLogicFrameTick();
+	PRELOADSCREEN_API void IOS_PlatformSpecificGameLogicFrameTick();
 #endif //PLATFORM_IOS
 };
diff --git a/Engine/Source/Runtime/PreLoadScreen/Public/PreLoadSettingsContainer.h b/Engine/Source/Runtime/PreLoadScreen/Public/PreLoadSettingsContainer.h
--- a/Engine/Source/Runtime/PreLoadScreen/Public/PreLoadSettingsContainer.h
+++ b/Engine/Source/Runtime/PreLoadScreen/Public/PreLoadSettingsContainer.h
@@ -16,7 +16,7 @@ struct FStandaloneCompositeFont;
 //This is a helper class that we use to hold values we parse from the .ini. Clean way to access things like dynamic image brushes / fonts / etc used in our UI that
 //we want to be somewhat data driven but we can't rely on UObject support to implement(as the PreLoad stuff happens too early for UObject support)
 //This lets us set easy to change values in our .ini that are parsed at runtime and stored in this container
-class PRELOADSCREEN_API FPreLoadSettingsContainerBase : public FDeferredCleanupInterface, public FGCObject
+class FPreLoadSettingsContainerBase : public FDeferredCleanupInterface, public FGCObject
 {
 public:
 
@@ -102,70 +102,70 @@ private:
 		HasCreatedSystemFontFile = false;
     }
 
-    virtual ~FPreLoadSettingsContainerBase();
+    PRELOADSCREEN_API virtual ~FPreLoadSettingsContainerBase();
 
 public:
 
 	//~ Begin FGCObject interface
-	virtual void AddReferencedObjects(FReferenceCollector& Collector) override;
+	PRELOADSCREEN_API virtual void AddReferencedObjects(FReferenceCollector& Collector) override;
 	virtual FString GetReferencerName() const override { return TEXT("FPreLoadSettingsContainerBase"); }
 	//~ End FGCObject interface
 
-    virtual const FSlateDynamicImageBrush* GetBrush(const FString& Identifier);
-    virtual FText GetLocalizedText(const FString& Identifier);
-    virtual TSharedPtr<FCompositeFont> GetFont(const FString& Identifier);
-    virtual FScreenGroupingBase* GetScreenGrouping(const FString& Identifier);
+    PRELOADSCREEN_API virtual const FSlateDynamicImageBrush* GetBrush(const FString& Identifier);
+    PRELOADSCREEN_API virtual FText GetLocalizedText(const FString& Identifier);
+    PRELOADSCREEN_API virtual TSharedPtr<FCompositeFont> GetFont(const FString& Identifier);
+    PRELOADSCREEN_API virtual FScreenGroupingBase* GetScreenGrouping(const FString& Identifier);
 
     int GetNumScreenGroupings() const { return ScreenGroupings.Num(); }
 
-	virtual const FScreenGroupingBase* GetScreenAtIndex(int index) const;
-	virtual bool IsValidScreenIndex(int index) const;
+	PRELOADSCREEN_API virtual const FScreenGroupingBase* GetScreenAtIndex(int index) const;
+	PRELOADSCREEN_API virtual bool IsValidScreenIndex(int index) const;
 
-    virtual void CreateCustomSlateImageBrush(const FString& Identifier, const FString& TexturePath, const FVector2D& ImageDimensions);
-    virtual void AddLocalizedText(const FString& Identifier, FText LocalizedText);
-    virtual void AddScreenGrouping(const FString& Identifier, FScreenGroupingBase& ScreenGrouping);
+    PRELOADSCREEN_API virtual void CreateCustomSlateImageBrush(const FString& Identifier, const FString& TexturePath, const FVector2D& ImageDimensions);
+    PRELOADSCREEN_API virtual void AddLocalizedText(const FString& Identifier, FText LocalizedText);
+    PRELOADSCREEN_API virtual void AddScreenGrouping(const FString& Identifier, FScreenGroupingBase& ScreenGrouping);
     
     //Maps the given font file to the given language and stores it under the FontIdentifier.
     //Identifier maps the entire CompositeFont, so if you want to add multiple fonts  for multiple languages, just store them all under the same identifer
-    virtual void BuildCustomFont(const FString& FontIdentifier, const FString& Language, const FString& FilePath);
-	virtual bool BuildSystemFontFile();
-	virtual const FString GetSystemFontFilePath() const;
+    PRELOADSCREEN_API virtual void BuildCustomFont(const FString& FontIdentifier, const FString& Language, const FString& FilePath);
+	PRELOADSCREEN_API virtual bool BuildSystemFontFile();
+	PRELOADSCREEN_API virtual const FString GetSystemFontFilePath() const;
 
     //Helper functions that parse a .ini config entry and call the appropriate create function to 
-    virtual void ParseBrushConfigEntry(const FString& BrushConfigEntry);
-    virtual void ParseFontConfigEntry(const FString&  SplitConfigEntry);
-    virtual void ParseLocalizedTextConfigString(const FString&  SplitConfigEntry);
-    virtual void ParseScreenGroupingConfigString(const FString&  SplitConfigEntry);
+    PRELOADSCREEN_API virtual void ParseBrushConfigEntry(const FString& BrushConfigEntry);
+    PRELOADSCREEN_API virtual void ParseFontConfigEntry(const FString&  SplitConfigEntry);
+    PRELOADSCREEN_API virtual void ParseLocalizedTextConfigString(const FString&  SplitConfigEntry);
+    PRELOADSCREEN_API virtual void ParseScreenGroupingConfigString(const FString&  SplitConfigEntry);
 
 	//Helper function to parse all .ini entries for LoadingGroups and ScreenOrder. Do these together so we can assert if
 	//we don't find a matching LoadingGroup identifier in the config. Should be run after we parse all screen groupings
-	virtual void ParseLoadingGroups(TArray<FString>& LoadingGroupIdentifiers);
-	virtual void ParseAllScreenOrderEntries(TArray<FString>& LoadingGroups, TArray<FString>& ScreenOrderEntries);
-	virtual void ParseScreenOrderConfigString(const FString& ScreenOrderEntry);
+	PRELOADSCREEN_API virtual void ParseLoadingGroups(TArray<FString>& LoadingGroupIdentifiers);
+	PRELOADSCREEN_API virtual void ParseAllScreenOrderEntries(TArray<FString>& LoadingGroups, TArray<FString>& ScreenOrderEntries);
+	PRELOADSCREEN_API virtual void ParseScreenOrderConfigString(const FString& ScreenOrderEntry);
 
     //Sets the PluginContent dir so that when parsing config entries we can accept plugin relative file paths
     virtual void SetPluginContentDir(const FString& PluginContentDirIn) { PluginContentDir = PluginContentDirIn; }
 
 	//Tells the container rather it should actually load image brushes
-	virtual void SetShouldLoadBrushes(bool bInShouldLoadBrushes);
+	PRELOADSCREEN_API virtual void SetShouldLoadBrushes(bool bInShouldLoadBrushes);
 
     float TimeToDisplayEachBackground;
     
 	FName GetCurrentLoadGrouping() const { return CurrentLoadGroup; }
-	void LoadGrouping(FName Identifier);
-	void PerformInitialAssetLoad();
+	PRELOADSCREEN_API void LoadGrouping(FName Identifier);
+	PRELOADSCREEN_API void PerformInitialAssetLoad();
 
     //Helper function that takes in a file path and tries to reconsile it to be Plugin Specific if applicable.
     //Ensures if file is not found in either Plugin's content dir or the original path
-    virtual FString ConvertIfPluginRelativeContentPath(const FString& FilePath);
+    PRELOADSCREEN_API virtual FString ConvertIfPluginRelativeContentPath(const FString& FilePath);
 
 protected:
 
     //Helper functions that verify if the supplied .ini config entry is valid to create a resource out of it
-    virtual bool IsValidBrushConfig(TArray<FString>& SplitConfigEntry);
-    virtual bool IsValidFontConfigString(TArray<FString>& SplitConfigEntry);
-    virtual bool IsValidLocalizedTextConfigString(TArray<FString>& SplitConfigEntry);
-    virtual bool IsValidScreenGrooupingConfigString(TArray<FString>& SplitConfigEntry);
+    PRELOADSCREEN_API virtual bool IsValidBrushConfig(TArray<FString>& SplitConfigEntry);
+    PRELOADSCREEN_API virtual bool IsValidFontConfigString(TArray<FString>& SplitConfigEntry);
+    PRELOADSCREEN_API virtual bool IsValidLocalizedTextConfigString(TArray<FString>& SplitConfigEntry);
+    PRELOADSCREEN_API virtual bool IsValidScreenGrooupingConfigString(TArray<FString>& SplitConfigEntry);
 
 protected:
 	TArray<FString> ParsedLoadingGroupIdentifiers;
@@ -191,12 +191,12 @@ protected:
 	bool HasCreatedSystemFontFile;
 
 	//If our Font filepath is set to this, we use the system font instead of a custom font we load in
-	static FString UseSystemFontOverride;
+	static PRELOADSCREEN_API FString UseSystemFontOverride;
 
 	//If we supply no loading groups, use this identifier by default
-	static FString DefaultInitialLoadingGroupIdentifier;
+	static PRELOADSCREEN_API FString DefaultInitialLoadingGroupIdentifier;
 
     // Singleton Instance -- This is only not a TSharedPtr as it needs to be cleaned up by a deferredcleanup call which directly
     // destroys the underlying object, causing a SharedPtr crash at shutdown.
-    static FPreLoadSettingsContainerBase* Instance;
-};
\ No newline at end of file
+    static PRELOADSCREEN_API FPreLoadSettingsContainerBase* Instance;
+};
diff --git a/Engine/Source/Runtime/PreLoadScreen/Public/PreLoadSlateThreading.h b/Engine/Source/Runtime/PreLoadScreen/Public/PreLoadSlateThreading.h
--- a/Engine/Source/Runtime/PreLoadScreen/Public/PreLoadSlateThreading.h
+++ b/Engine/Source/Runtime/PreLoadScreen/Public/PreLoadSlateThreading.h
@@ -45,12 +45,12 @@ private:
     class FPreLoadScreenSlateSynchMechanism* SyncMechanism;
 };
 
-class PRELOADSCREEN_API FPreLoadSlateWidgetRenderer
+class FPreLoadSlateWidgetRenderer
 {
 public:
-	FPreLoadSlateWidgetRenderer(TSharedPtr<SWindow> InMainWindow, TSharedPtr<SVirtualWindow> InVirtualRenderWindowWindow, FSlateRenderer* InRenderer);
+	PRELOADSCREEN_API FPreLoadSlateWidgetRenderer(TSharedPtr<SWindow> InMainWindow, TSharedPtr<SVirtualWindow> InVirtualRenderWindowWindow, FSlateRenderer* InRenderer);
 
-	void DrawWindow(float DeltaTime);
+	PRELOADSCREEN_API void DrawWindow(float DeltaTime);
 
 	SWindow* GetMainWindow_GameThread() const { return MainWindow; }
 
@@ -74,37 +74,37 @@ private:
  * This class will handle all the nasty bits about running Slate on a separate thread
  * and then trying to sync it up with the game thread and the render thread simultaneously
  */
-class PRELOADSCREEN_API FPreLoadScreenSlateSynchMechanism
+class FPreLoadScreenSlateSynchMechanism
 {
 public:
-    FPreLoadScreenSlateSynchMechanism(TSharedPtr<FPreLoadSlateWidgetRenderer, ESPMode::ThreadSafe> InWidgetRenderer);
-    ~FPreLoadScreenSlateSynchMechanism();
+    PRELOADSCREEN_API FPreLoadScreenSlateSynchMechanism(TSharedPtr<FPreLoadSlateWidgetRenderer, ESPMode::ThreadSafe> InWidgetRenderer);
+    PRELOADSCREEN_API ~FPreLoadScreenSlateSynchMechanism();
 
 	FPreLoadScreenSlateSynchMechanism() = delete;
 	FPreLoadScreenSlateSynchMechanism(const FPreLoadScreenSlateSynchMechanism&) = delete;
 	FPreLoadScreenSlateSynchMechanism& operator=(const FPreLoadScreenSlateSynchMechanism&) = delete;
 
     /** Sets up the locks in their proper initial state for running */
-    void Initialize();
+    PRELOADSCREEN_API void Initialize();
 
     /** Cleans up the slate thread */
-    void DestroySlateThread();
+    PRELOADSCREEN_API void DestroySlateThread();
 
     /** Handles the counter to determine if the slate thread should keep running */
-    bool IsSlateMainLoopRunning_AnyThread() const;
+    PRELOADSCREEN_API bool IsSlateMainLoopRunning_AnyThread() const;
 
 private:
 	/** Notified when a SWindow is being destroyed */
-	void HandleWindowBeingDestroyed(const SWindow& WindowBeingDestroyed);
+	PRELOADSCREEN_API void HandleWindowBeingDestroyed(const SWindow& WindowBeingDestroyed);
 
 	/** The main loop to be run from the Slate thread */
-	void RunMainLoop_SlateThread();
+	PRELOADSCREEN_API void RunMainLoop_SlateThread();
 
     /** This counter handles running the main loop of the slate thread */
     TAtomic<bool> bIsRunningSlateMainLoop;
 
     /** This counter is used to generate a unique id for each new instance of the loading thread */
-    static TAtomic<int32> LoadingThreadInstanceCounter;
+    static PRELOADSCREEN_API TAtomic<int32> LoadingThreadInstanceCounter;
 
     /** The worker thread that will become the Slate thread */
     FRunnableThread* SlateLoadingThread;
diff --git a/Engine/Source/Runtime/Projects/Public/CustomBuildSteps.h b/Engine/Source/Runtime/Projects/Public/CustomBuildSteps.h
--- a/Engine/Source/Runtime/Projects/Public/CustomBuildSteps.h
+++ b/Engine/Source/Runtime/Projects/Public/CustomBuildSteps.h
@@ -13,20 +13,20 @@ class FJsonObject;
 /**
  * Descriptor for projects. Contains all the information contained within a .uproject file.
  */
-struct PROJECTS_API FCustomBuildSteps
+struct FCustomBuildSteps
 {
 	/** Mapping from host platform to list of commands */
 	TMap<FString, TArray<FString>> HostPlatformToCommands;
 
 	/** Whether this object is empty */
-	bool IsEmpty() const;
+	PROJECTS_API bool IsEmpty() const;
 
 	/** Reads the descriptor from the given JSON object */
-	void Read(const FJsonObject& Object, const FString& FieldName);
+	PROJECTS_API void Read(const FJsonObject& Object, const FString& FieldName);
 
 	/** Writes the CustomBuildSteps to JSON */
-	void Write(TJsonWriter<>& Writer, const FString& FieldName) const;
+	PROJECTS_API void Write(TJsonWriter<>& Writer, const FString& FieldName) const;
 
 	/** Updates the given json object with this CustomBuildSteps */
-	void UpdateJson(FJsonObject& JsonObject, const FString& FieldName) const;
+	PROJECTS_API void UpdateJson(FJsonObject& JsonObject, const FString& FieldName) const;
 };
diff --git a/Engine/Source/Runtime/Projects/Public/LocalizationDescriptor.h b/Engine/Source/Runtime/Projects/Public/LocalizationDescriptor.h
--- a/Engine/Source/Runtime/Projects/Public/LocalizationDescriptor.h
+++ b/Engine/Source/Runtime/Projects/Public/LocalizationDescriptor.h
@@ -54,7 +54,7 @@ namespace ELocalizationTargetDescriptorLoadingPolicy
 /**
  * Description of a localization target.
  */
-struct PROJECTS_API FLocalizationTargetDescriptor
+struct FLocalizationTargetDescriptor
 {
 	/** Name of this target */
 	FString Name;
@@ -63,32 +63,32 @@ struct PROJECTS_API FLocalizationTargetDescriptor
 	ELocalizationTargetDescriptorLoadingPolicy::Type LoadingPolicy;
 
 	/** Normal constructor */
-	FLocalizationTargetDescriptor(FString InName = FString(), ELocalizationTargetDescriptorLoadingPolicy::Type InLoadingPolicy = ELocalizationTargetDescriptorLoadingPolicy::Never);
+	PROJECTS_API FLocalizationTargetDescriptor(FString InName = FString(), ELocalizationTargetDescriptorLoadingPolicy::Type InLoadingPolicy = ELocalizationTargetDescriptorLoadingPolicy::Never);
 
 	/** Reads a descriptor from the given JSON object */
-	bool Read(const FJsonObject& InObject, FText* OutFailReason = nullptr);
+	PROJECTS_API bool Read(const FJsonObject& InObject, FText* OutFailReason = nullptr);
 
 	/** Reads a descriptor from the given JSON object */
-	bool Read(const FJsonObject& InObject, FText& OutFailReason);
+	PROJECTS_API bool Read(const FJsonObject& InObject, FText& OutFailReason);
 
 	/** Reads an array of targets from the given JSON object */
-	static bool ReadArray(const FJsonObject& InObject, const TCHAR* InName, TArray<FLocalizationTargetDescriptor>& OutTargets, FText* OutFailReason = nullptr);
+	static PROJECTS_API bool ReadArray(const FJsonObject& InObject, const TCHAR* InName, TArray<FLocalizationTargetDescriptor>& OutTargets, FText* OutFailReason = nullptr);
 
 	/** Reads an array of targets from the given JSON object */
-	static bool ReadArray(const FJsonObject& InObject, const TCHAR* InName, TArray<FLocalizationTargetDescriptor>& OutTargets, FText& OutFailReason);
+	static PROJECTS_API bool ReadArray(const FJsonObject& InObject, const TCHAR* InName, TArray<FLocalizationTargetDescriptor>& OutTargets, FText& OutFailReason);
 
 	/** Writes a descriptor to JSON */
-	void Write(TJsonWriter<>& Writer) const;
+	PROJECTS_API void Write(TJsonWriter<>& Writer) const;
 
 	/** Updates the given json object with values in this descriptor */
-	void UpdateJson(FJsonObject& JsonObject) const;
+	PROJECTS_API void UpdateJson(FJsonObject& JsonObject) const;
 
 	/** Writes an array of targets to JSON */
-	static void WriteArray(TJsonWriter<>& Writer, const TCHAR* ArrayName, const TArray<FLocalizationTargetDescriptor>& Descriptors);
+	static PROJECTS_API void WriteArray(TJsonWriter<>& Writer, const TCHAR* ArrayName, const TArray<FLocalizationTargetDescriptor>& Descriptors);
 
 	/** Updates an array of descriptors in the specified JSON field (indexed by name) */
-	static void UpdateArray(FJsonObject& JsonObject, const TCHAR* ArrayName, const TArray<FLocalizationTargetDescriptor>& Descriptors);
+	static PROJECTS_API void UpdateArray(FJsonObject& JsonObject, const TCHAR* ArrayName, const TArray<FLocalizationTargetDescriptor>& Descriptors);
 
 	/** Returns true if we should load this localization target based upon the current runtime environment */
-	bool ShouldLoadLocalizationTarget() const;
+	PROJECTS_API bool ShouldLoadLocalizationTarget() const;
 };
diff --git a/Engine/Source/Runtime/Projects/Public/ModuleDescriptor.h b/Engine/Source/Runtime/Projects/Public/ModuleDescriptor.h
--- a/Engine/Source/Runtime/Projects/Public/ModuleDescriptor.h
+++ b/Engine/Source/Runtime/Projects/Public/ModuleDescriptor.h
@@ -150,7 +150,7 @@ namespace EHostType
 /**
  * Description of a loadable module.
  */
-struct PROJECTS_API FModuleDescriptor
+struct FModuleDescriptor
 {
 	/** Name of this module */
 	FName Name;
@@ -193,47 +193,47 @@ struct PROJECTS_API FModuleDescriptor
 
 
 	/** Normal constructor */
-	FModuleDescriptor(const FName InName = NAME_None, EHostType::Type InType = EHostType::Runtime, ELoadingPhase::Type InLoadingPhase = ELoadingPhase::Default);
+	PROJECTS_API FModuleDescriptor(const FName InName = NAME_None, EHostType::Type InType = EHostType::Runtime, ELoadingPhase::Type InLoadingPhase = ELoadingPhase::Default);
 
 	/** Reads a descriptor from the given JSON object */
-	bool Read(const FJsonObject& Object, FText* OutFailReason = nullptr);
+	PROJECTS_API bool Read(const FJsonObject& Object, FText* OutFailReason = nullptr);
 
 	/** Reads a descriptor from the given JSON object */
-	bool Read(const FJsonObject& Object, FText& OutFailReason);
+	PROJECTS_API bool Read(const FJsonObject& Object, FText& OutFailReason);
 
 	/** Reads an array of modules from the given JSON object */
-	static bool ReadArray(const FJsonObject& Object, const TCHAR* Name, TArray<FModuleDescriptor>& OutModules, FText* OutFailReason = nullptr);
+	static PROJECTS_API bool ReadArray(const FJsonObject& Object, const TCHAR* Name, TArray<FModuleDescriptor>& OutModules, FText* OutFailReason = nullptr);
 
 	/** Reads an array of modules from the given JSON object */
-	static bool ReadArray(const FJsonObject& Object, const TCHAR* Name, TArray<FModuleDescriptor>& OutModules, FText& OutFailReason);
+	static PROJECTS_API bool ReadArray(const FJsonObject& Object, const TCHAR* Name, TArray<FModuleDescriptor>& OutModules, FText& OutFailReason);
 
 	/** Writes a descriptor to JSON */
-	void Write(TJsonWriter<>& Writer) const;
+	PROJECTS_API void Write(TJsonWriter<>& Writer) const;
 
 	/** Updates the given json object with values in this descriptor */
-	void UpdateJson(FJsonObject& JsonObject) const;
+	PROJECTS_API void UpdateJson(FJsonObject& JsonObject) const;
 
 	/** Writes an array of modules to JSON */
-	static void WriteArray(TJsonWriter<>& Writer, const TCHAR* ArrayName, const TArray<FModuleDescriptor>& Modules);
+	static PROJECTS_API void WriteArray(TJsonWriter<>& Writer, const TCHAR* ArrayName, const TArray<FModuleDescriptor>& Modules);
 
 	/** Updates an array of module descriptors in the specified JSON field (indexed by module name) */
-	static void UpdateArray(FJsonObject& JsonObject, const TCHAR* ArrayName, const TArray<FModuleDescriptor>& Modules);
+	static PROJECTS_API void UpdateArray(FJsonObject& JsonObject, const TCHAR* ArrayName, const TArray<FModuleDescriptor>& Modules);
 
 	/** Tests whether the module should be built for the given target */
-	bool IsCompiledInConfiguration(const FString& Platform, EBuildConfiguration Configuration, const FString& TargetName, EBuildTargetType TargetType, bool bBuildDeveloperTools, bool bBuildRequiresCookedData) const;
+	PROJECTS_API bool IsCompiledInConfiguration(const FString& Platform, EBuildConfiguration Configuration, const FString& TargetName, EBuildTargetType TargetType, bool bBuildDeveloperTools, bool bBuildRequiresCookedData) const;
 
 	/** Tests whether the module should be built for the current engine configuration */
-	bool IsCompiledInCurrentConfiguration() const;
+	PROJECTS_API bool IsCompiledInCurrentConfiguration() const;
 
 	/** Tests whether the module should be loaded for the current engine configuration */
-	bool IsLoadedInCurrentConfiguration() const;
+	PROJECTS_API bool IsLoadedInCurrentConfiguration() const;
 
 	/** Loads all the modules for a given loading phase. Returns a map of module names to load errors */
-	static void LoadModulesForPhase(ELoadingPhase::Type LoadingPhase, const TArray<FModuleDescriptor>& Modules, TMap<FName, EModuleLoadResult>& ModuleLoadErrors);
+	static PROJECTS_API void LoadModulesForPhase(ELoadingPhase::Type LoadingPhase, const TArray<FModuleDescriptor>& Modules, TMap<FName, EModuleLoadResult>& ModuleLoadErrors);
 
 #if !IS_MONOLITHIC
 	/** Checks that all modules are compatible with the current engine version. Returns false and appends a list of names to OutIncompatibleFiles if not. */
-	static bool CheckModuleCompatibility(const TArray<FModuleDescriptor>& Modules, TArray<FString>& OutIncompatibleFiles);
+	static PROJECTS_API bool CheckModuleCompatibility(const TArray<FModuleDescriptor>& Modules, TArray<FString>& OutIncompatibleFiles);
 #endif
 };
 
diff --git a/Engine/Source/Runtime/Projects/Public/PluginDescriptor.h b/Engine/Source/Runtime/Projects/Public/PluginDescriptor.h
--- a/Engine/Source/Runtime/Projects/Public/PluginDescriptor.h
+++ b/Engine/Source/Runtime/Projects/Public/PluginDescriptor.h
@@ -32,7 +32,7 @@ enum class EPluginEnabledByDefault : uint8
 /**
  * Descriptor for plugins. Contains all the information contained within a .uplugin file.
  */
-struct PROJECTS_API FPluginDescriptor
+struct FPluginDescriptor
 {
 	/** Version number for the plugin.  The version number must increase with every version of the plugin, so that the system 
 	    can determine whether one version of a plugin is newer than another, or to enforce other requirements.  This version
@@ -158,57 +158,57 @@ struct PROJECTS_API FPluginDescriptor
 #endif
 
 	/** Return the .uplugin extension (with dot) */
-	static const FString& GetFileExtension();
+	static PROJECTS_API const FString& GetFileExtension();
 
 	/** Constructor. */
-	FPluginDescriptor();
+	PROJECTS_API FPluginDescriptor();
 
 	/** Loads the descriptor from the given file. */
-	bool Load(const FString& FileName, FText* OutFailReason = nullptr);
+	PROJECTS_API bool Load(const FString& FileName, FText* OutFailReason = nullptr);
 
 	/** Loads the descriptor from the given file. */
-	bool Load(const FString& FileName, FText& OutFailReason);
+	PROJECTS_API bool Load(const FString& FileName, FText& OutFailReason);
 
 	/** Reads the descriptor from the given string */
-	bool Read(const FString& Text, FText* OutFailReason = nullptr);
+	PROJECTS_API bool Read(const FString& Text, FText* OutFailReason = nullptr);
 
 	/** Reads the descriptor from the given string */
-	bool Read(const FString& Text, FText& OutFailReason);
+	PROJECTS_API bool Read(const FString& Text, FText& OutFailReason);
 
 	/** Reads the descriptor from the given JSON object */
-	bool Read(const FJsonObject& Object, FText* OutFailReason = nullptr);
+	PROJECTS_API bool Read(const FJsonObject& Object, FText* OutFailReason = nullptr);
 
 	/** Reads the descriptor from the given JSON object */
-	bool Read(const FJsonObject& Object, FText& OutFailReason);
+	PROJECTS_API bool Read(const FJsonObject& Object, FText& OutFailReason);
 
 	/** Saves the descriptor from the given file. */
-	bool Save(const FString& FileName, FText* OutFailReason = nullptr) const;
+	PROJECTS_API bool Save(const FString& FileName, FText* OutFailReason = nullptr) const;
 
 	/** Saves the descriptor from the given file. */
-	bool Save(const FString& FileName, FText& OutFailReason) const;
+	PROJECTS_API bool Save(const FString& FileName, FText& OutFailReason) const;
 
 	/** Writes a descriptor to JSON */
-	void Write(FString& Text) const;
+	PROJECTS_API void Write(FString& Text) const;
 
 	/** Writes a descriptor to JSON */
-	void Write(TJsonWriter<>& Writer) const;
+	PROJECTS_API void Write(TJsonWriter<>& Writer) const;
 
 	/** Updates the given json object with values in this descriptor */
-	void UpdateJson(FJsonObject& JsonObject) const;
+	PROJECTS_API void UpdateJson(FJsonObject& JsonObject) const;
 
 	/**
 	 * Updates the content of the specified plugin file with values in this descriptor
 	 * (hence preserving json fields that the plugin descriptor doesn't know about)
 	 */
-	bool UpdatePluginFile(const FString& FileName, FText* OutFailReason = nullptr) const;
+	PROJECTS_API bool UpdatePluginFile(const FString& FileName, FText* OutFailReason = nullptr) const;
 
 	/**
 	 * Updates the content of the specified plugin file with values in this descriptor
 	 * (hence preserving json fields that the plugin descriptor doesn't know about)
 	 */
-	bool UpdatePluginFile(const FString& FileName, FText& OutFailReason) const;
+	PROJECTS_API bool UpdatePluginFile(const FString& FileName, FText& OutFailReason) const;
 
 	/** Determines whether the plugin supports the given platform */
-	bool SupportsTargetPlatform(const FString& Platform) const;
+	PROJECTS_API bool SupportsTargetPlatform(const FString& Platform) const;
 };
 
diff --git a/Engine/Source/Runtime/Projects/Public/PluginManifest.h b/Engine/Source/Runtime/Projects/Public/PluginManifest.h
--- a/Engine/Source/Runtime/Projects/Public/PluginManifest.h
+++ b/Engine/Source/Runtime/Projects/Public/PluginManifest.h
@@ -13,7 +13,7 @@ class FText;
 /**
  * Descriptor for plugins. Contains all the information contained within a .uplugin file.
  */
-struct PROJECTS_API FPluginManifestEntry
+struct FPluginManifestEntry
 {
 	/** Normalized path to the plugin file */
 	FString File;
@@ -25,14 +25,14 @@ struct PROJECTS_API FPluginManifestEntry
 /**
  * Manifest of plugins. Descriptor for plugins. Contains all the information contained within a .uplugin file.
  */
-struct PROJECTS_API FPluginManifest
+struct FPluginManifest
 {
 	/** List of plugins in this manifest */
 	TArray<FPluginManifestEntry> Contents;
 
 	/** Loads the descriptor from the given file. */
-	bool Load(const FString& FileName, FText& OutFailReason);
+	PROJECTS_API bool Load(const FString& FileName, FText& OutFailReason);
 
 	/** Reads the descriptor from the given JSON object */
-	bool Read(const FJsonObject& Object, FText& OutFailReason);
-};
\ No newline at end of file
+	PROJECTS_API bool Read(const FJsonObject& Object, FText& OutFailReason);
+};
diff --git a/Engine/Source/Runtime/Projects/Public/PluginReferenceDescriptor.h b/Engine/Source/Runtime/Projects/Public/PluginReferenceDescriptor.h
--- a/Engine/Source/Runtime/Projects/Public/PluginReferenceDescriptor.h
+++ b/Engine/Source/Runtime/Projects/Public/PluginReferenceDescriptor.h
@@ -23,7 +23,7 @@ class FText;
 /**
  * Descriptor for a plugin reference. Contains the information required to enable or disable a plugin for a given platform.
  */
-struct PROJECTS_API FPluginReferenceDescriptor
+struct FPluginReferenceDescriptor
 {
 	/** Name of the plugin */
 	FString Name;
@@ -76,49 +76,49 @@ struct PROJECTS_API FPluginReferenceDescriptor
 #endif //if WITH_EDITOR
 
 	/** Constructor */
-	FPluginReferenceDescriptor(const FString& InName = TEXT(""), bool bInEnabled = false);
+	PROJECTS_API FPluginReferenceDescriptor(const FString& InName = TEXT(""), bool bInEnabled = false);
 
 	/** Determines whether the plugin is enabled for the given platform */
-	bool IsEnabledForPlatform(const FString& Platform) const;
+	PROJECTS_API bool IsEnabledForPlatform(const FString& Platform) const;
 
 	/** Determines whether the plugin is enabled for the given target configuration */
-	bool IsEnabledForTargetConfiguration(EBuildConfiguration Configuration) const;
+	PROJECTS_API bool IsEnabledForTargetConfiguration(EBuildConfiguration Configuration) const;
 
 	/** Determines whether the plugin is enabled for the given target */
-	bool IsEnabledForTarget(EBuildTargetType TargetType) const;
+	PROJECTS_API bool IsEnabledForTarget(EBuildTargetType TargetType) const;
 
 	/** Determines if the referenced plugin is supported for the given platform */
-	bool IsSupportedTargetPlatform(const FString& Platform) const;
+	PROJECTS_API bool IsSupportedTargetPlatform(const FString& Platform) const;
 
 	/** Reads the descriptor from the given JSON object */
-	bool Read(const TSharedRef<FJsonObject>& Object, FText* OutFailReason = nullptr);
+	PROJECTS_API bool Read(const TSharedRef<FJsonObject>& Object, FText* OutFailReason = nullptr);
 
 	UE_DEPRECATED(5.1, "Use Read(const TSharedRef<FJsonObject>&) instead.")
-	bool Read(const FJsonObject& Object, FText* OutFailReason = nullptr, TSharedPtr<FJsonObject> ObjectPtr = nullptr);
+	PROJECTS_API bool Read(const FJsonObject& Object, FText* OutFailReason = nullptr, TSharedPtr<FJsonObject> ObjectPtr = nullptr);
 
 	UE_DEPRECATED(5.1, "Use Read(const TSharedRef<FJsonObject>&) instead.")
-	bool Read(const FJsonObject& Object, FText& OutFailReason, TSharedPtr<FJsonObject> ObjectPtr = nullptr);
+	PROJECTS_API bool Read(const FJsonObject& Object, FText& OutFailReason, TSharedPtr<FJsonObject> ObjectPtr = nullptr);
 
 	/** Reads an array of modules from the given JSON object */
-	static bool ReadArray(const FJsonObject& Object, const TCHAR* Name, TArray<FPluginReferenceDescriptor>& OutModules, FText* OutFailReason = nullptr);
+	static PROJECTS_API bool ReadArray(const FJsonObject& Object, const TCHAR* Name, TArray<FPluginReferenceDescriptor>& OutModules, FText* OutFailReason = nullptr);
 
 	/** Reads an array of modules from the given JSON object */
-	static bool ReadArray(const FJsonObject& Object, const TCHAR* Name, TArray<FPluginReferenceDescriptor>& OutModules, FText& OutFailReason);
+	static PROJECTS_API bool ReadArray(const FJsonObject& Object, const TCHAR* Name, TArray<FPluginReferenceDescriptor>& OutModules, FText& OutFailReason);
 
 	/** Writes a descriptor to JSON */
-	void Write(TJsonWriter<>& Writer) const;
+	PROJECTS_API void Write(TJsonWriter<>& Writer) const;
 	
 	/** Updates the given json object with values in this descriptor */
-	void UpdateJson(FJsonObject& JsonObject) const;
+	PROJECTS_API void UpdateJson(FJsonObject& JsonObject) const;
 
 	/** Writes an array of plugin references to JSON */
-	static void WriteArray(TJsonWriter<>& Writer, const TCHAR* ArrayName, const TArray<FPluginReferenceDescriptor>& Plugins);
+	static PROJECTS_API void WriteArray(TJsonWriter<>& Writer, const TCHAR* ArrayName, const TArray<FPluginReferenceDescriptor>& Plugins);
 
 	/** Updates an array of plugin references in the specified JSON field (indexed by plugin name) */
-	static void UpdateArray(FJsonObject& JsonObject, const TCHAR* ArrayName, const TArray<FPluginReferenceDescriptor>& Plugins);
+	static PROJECTS_API void UpdateArray(FJsonObject& JsonObject, const TCHAR* ArrayName, const TArray<FPluginReferenceDescriptor>& Plugins);
 
 #if WITH_EDITOR
 	/** Gets the string value for a given key by first looking into AdditionalFieldsToWrite and then CachedJson */
-	bool GetAdditionalStringField(const FString& Key, FString& OutValue) const;
+	PROJECTS_API bool GetAdditionalStringField(const FString& Key, FString& OutValue) const;
 #endif //if WITH_EDITOR
 };
diff --git a/Engine/Source/Runtime/Projects/Public/ProjectDescriptor.h b/Engine/Source/Runtime/Projects/Public/ProjectDescriptor.h
--- a/Engine/Source/Runtime/Projects/Public/ProjectDescriptor.h
+++ b/Engine/Source/Runtime/Projects/Public/ProjectDescriptor.h
@@ -38,7 +38,7 @@ namespace EProjectDescriptorVersion
 /**
  * Descriptor for projects. Contains all the information contained within a .uproject file.
  */
-struct PROJECTS_API FProjectDescriptor
+struct FProjectDescriptor
 {
 	/** Descriptor version number. */
 	EProjectDescriptorVersion::Type FileVersion;
@@ -105,34 +105,34 @@ struct PROJECTS_API FProjectDescriptor
 	bool bDisableEnginePluginsByDefault;
 
 	/** Constructor. */
-	FProjectDescriptor();
+	PROJECTS_API FProjectDescriptor();
 
 	/** Signs the project given for the given filename */
-	void Sign(const FString& FilePath);
+	PROJECTS_API void Sign(const FString& FilePath);
 
 	/** Checks whether the descriptor is signed */
-	bool IsSigned(const FString& FilePath) const;
+	PROJECTS_API bool IsSigned(const FString& FilePath) const;
 
 	/** Finds the index of a plugin in the references array */
-	int32 FindPluginReferenceIndex(const FString& PluginName) const;
+	PROJECTS_API int32 FindPluginReferenceIndex(const FString& PluginName) const;
 
 	/** Updates the supported target platforms list */
-	void UpdateSupportedTargetPlatforms(const FName& InPlatformName, bool bIsSupported);
+	PROJECTS_API void UpdateSupportedTargetPlatforms(const FName& InPlatformName, bool bIsSupported);
 
 	/** Loads the descriptor from the given file. */
-	bool Load(const FString& FileName, FText& OutFailReason);
+	PROJECTS_API bool Load(const FString& FileName, FText& OutFailReason);
 
 	/** Reads the descriptor from the given JSON object */
-	bool Read(const FJsonObject& Object, const FString& PathToProject, FText& OutFailReason);
+	PROJECTS_API bool Read(const FJsonObject& Object, const FString& PathToProject, FText& OutFailReason);
 
 	/** Saves the descriptor to the given file. */
-	bool Save(const FString& FileName, FText& OutFailReason);
+	PROJECTS_API bool Save(const FString& FileName, FText& OutFailReason);
 
 	/** Writes the descriptor to the given JSON object */
-	void Write(TJsonWriter<>& Writer, const FString& PathToProject) const;
+	PROJECTS_API void Write(TJsonWriter<>& Writer, const FString& PathToProject) const;
 
 	/** Returns the extension used for project descriptors (uproject) */
-	static FString GetExtension();
+	static PROJECTS_API FString GetExtension();
 
 	/** @return - Access to the additional plugin directories */
 	const TArray<FString>& GetAdditionalPluginDirectories() const
@@ -146,14 +146,14 @@ struct PROJECTS_API FProjectDescriptor
 	 * @param Dir - the new directory to add (must be a full path)
 	 * @return whether the plugin directory list was changed
 	 */
-	bool AddPluginDirectory(const FString& Dir);
+	PROJECTS_API bool AddPluginDirectory(const FString& Dir);
 	/**
 	 * Removes the directory from the list to scan
 	 *
 	 * @param Dir the directory to remove (must be a full path)
 	 * @return whether the plugin directory list was changed
 	 */
-	bool RemovePluginDirectory(const FString& Dir);
+	PROJECTS_API bool RemovePluginDirectory(const FString& Dir);
 
 	/** @return - Access to the additional root directories */
 	const TArray<FString>& GetAdditionalRootDirectories() const
@@ -167,23 +167,23 @@ struct PROJECTS_API FProjectDescriptor
 	 * @param Dir - the new directory to add (must be a full path)
 	 * @return whether the root directory list was changed
 	 */
-	bool AddRootDirectory(const FString& Dir);
+	PROJECTS_API bool AddRootDirectory(const FString& Dir);
 	/**
 	 * Removes the directory from the list to scan
 	 *
 	 * @param Dir the directory to remove (must be a full path)
 	 * @return whether the root directory list was changed
 	 */
-	bool RemoveRootDirectory(const FString& Dir);
+	PROJECTS_API bool RemoveRootDirectory(const FString& Dir);
 
 #if WITH_EDITOR
 	/** Returns whether the project has a module of the given name */
-	bool HasModule(FName ModuleName) const;
+	PROJECTS_API bool HasModule(FName ModuleName) const;
 #endif //if WITH_EDITOR
 
 private:
 	/** @return the path relative to this project if possible */
-	const FString MakePathRelativeToProject(const FString& Dir, const FString& PathToProject) const;
+	PROJECTS_API const FString MakePathRelativeToProject(const FString& Dir, const FString& PathToProject) const;
 
 	/**
 	 * List of additional directories to scan for plugins.
diff --git a/Engine/Source/Runtime/PropertyPath/Public/PropertyPathHelpers.h b/Engine/Source/Runtime/PropertyPath/Public/PropertyPathHelpers.h
--- a/Engine/Source/Runtime/PropertyPath/Public/PropertyPathHelpers.h
+++ b/Engine/Source/Runtime/PropertyPath/Public/PropertyPathHelpers.h
@@ -24,43 +24,43 @@ template <typename FuncType> class TFunctionRef;
 
 /** Base class for cached property path segments */
 USTRUCT()
-struct PROPERTYPATH_API FPropertyPathSegment
+struct FPropertyPathSegment
 {
 	GENERATED_BODY()
 
 	friend struct FCachedPropertyPath;
 
 	/** UStruct boilerplate constructor - do not use */
-	FPropertyPathSegment();
+	PROPERTYPATH_API FPropertyPathSegment();
 
 	/** Construct from char count and storage */
-	FPropertyPathSegment(int32 InCount, const TCHAR* InString);
+	PROPERTYPATH_API FPropertyPathSegment(int32 InCount, const TCHAR* InString);
 
 	/**
 	 * Called after this event has been serialized in order to cache the field pointer if necessary
 	 */
-	void PostSerialize(const FArchive& Ar);
+	PROPERTYPATH_API void PostSerialize(const FArchive& Ar);
 
 	/** Make a copy which is unresolved */
-	static FPropertyPathSegment MakeUnresolvedCopy(const FPropertyPathSegment& ToCopy);
+	static PROPERTYPATH_API FPropertyPathSegment MakeUnresolvedCopy(const FPropertyPathSegment& ToCopy);
 
 	/**
 	 * Resolves the name on the given Struct.  Can be used to cache the resulting property so that future calls can be processed quickly.
 	 * @param InStruct the ScriptStruct or Class to look for the property on.
 	 */
-	FFieldVariant Resolve(UStruct* InStruct) const;
+	PROPERTYPATH_API FFieldVariant Resolve(UStruct* InStruct) const;
 
 	/** @return the name of this segment */
-	FName GetName() const;
+	PROPERTYPATH_API FName GetName() const;
 
 	/** @return the array index of this segment */
-	int32 GetArrayIndex() const;
+	PROPERTYPATH_API int32 GetArrayIndex() const;
 
 	/** @return the resolved field */
-	FFieldVariant GetField() const;
+	PROPERTYPATH_API FFieldVariant GetField() const;
 
 	/** @return the resolved struct */
-	UStruct* GetStruct() const;
+	PROPERTYPATH_API UStruct* GetStruct() const;
 
 public:
 
@@ -93,123 +93,123 @@ struct TStructOpsTypeTraits<FPropertyPathSegment> : TStructOpsTypeTraitsBase2<FP
 
 /** Base class for cached property paths */
 USTRUCT()
-struct PROPERTYPATH_API FCachedPropertyPath
+struct FCachedPropertyPath
 {
 	GENERATED_BODY()
 
 	/** UStruct boilerplate constructor - do not use */
-	FCachedPropertyPath();
+	PROPERTYPATH_API FCachedPropertyPath();
 
 	/** */
-	FCachedPropertyPath(const FString& Path);
+	PROPERTYPATH_API FCachedPropertyPath(const FString& Path);
 
 	/** */
-	FCachedPropertyPath(const FPropertyPathSegment& Segment);
+	PROPERTYPATH_API FCachedPropertyPath(const FPropertyPathSegment& Segment);
 
 	/** */
-	FCachedPropertyPath(const TArray<FString>& PathSegments);
+	PROPERTYPATH_API FCachedPropertyPath(const TArray<FString>& PathSegments);
 
 	/** */
-	~FCachedPropertyPath();
+	PROPERTYPATH_API ~FCachedPropertyPath();
 
 	/** Check whether this property path is non-empty */
 	bool IsValid() const { return Segments.Num() > 0; }
 
 	/** Make a new property path from a string */
-	void MakeFromString(const FString& InPropertyPath);
+	PROPERTYPATH_API void MakeFromString(const FString& InPropertyPath);
 
 	/** Make a copy which is unresolved */
-	static FCachedPropertyPath MakeUnresolvedCopy(const FCachedPropertyPath& ToCopy);
+	static PROPERTYPATH_API FCachedPropertyPath MakeUnresolvedCopy(const FCachedPropertyPath& ToCopy);
 
 	/** @return Get the number of segments in this path */
-	int32 GetNumSegments() const;
+	PROPERTYPATH_API int32 GetNumSegments() const;
 
 	/** 
 	 * Get the path segment at the specified index
 	 * @param	InSegmentIndex	The index of the segment
 	 * @return the segment at the specified index 
 	 */
-	const FPropertyPathSegment& GetSegment(int32 InSegmentIndex) const;
+	PROPERTYPATH_API const FPropertyPathSegment& GetSegment(int32 InSegmentIndex) const;
 
 	/** 
 	 * Get the path segment at the end of the path
 	 * @return the segment at the specified index 
 	 */
-	const FPropertyPathSegment& GetLastSegment() const;
+	PROPERTYPATH_API const FPropertyPathSegment& GetLastSegment() const;
 
 	/** 
 	 * Resolve this property path against the specified object.
 	 * @return true if the path could be resolved
 	 */
-	bool Resolve(UObject* InContainer) const;
+	PROPERTYPATH_API bool Resolve(UObject* InContainer) const;
 
 	/** Set whether this path resolves over object or dynamic array boundaries, making it unsafe for general direct cached access */
-	void SetCanSafelyUsedCachedAddress(bool bInCanSafelyUsedCachedAddress) const;
+	PROPERTYPATH_API void SetCanSafelyUsedCachedAddress(bool bInCanSafelyUsedCachedAddress) const;
 
 	/** Update cached last container property in path & correspondng index, invalidates 'bCanSafelyUsedCachedAddress' */
-	void SetCachedLastContainer(void* InContainer, int32 InIndex) const;
+	PROPERTYPATH_API void SetCachedLastContainer(void* InContainer, int32 InIndex) const;
 
 	/** Get cached last container property in path */
-	void* GetCachedLastContainerInPath() const;
+	PROPERTYPATH_API void* GetCachedLastContainerInPath() const;
 
 	/** Get cached index of last container property in path, INDEX_NONE if last container is not in path */
-	int32 GetCachedLastContainerInPathIndex() const;
+	PROPERTYPATH_API int32 GetCachedLastContainerInPathIndex() const;
 
 	/** Cache a resolved address for faster subsequent access */
-	void ResolveLeaf(void* InAddress) const;
+	PROPERTYPATH_API void ResolveLeaf(void* InAddress) const;
 
 	/** Cache a resolved function for faster subsequent access */
-	void ResolveLeaf(UFunction* InFunction) const;
+	PROPERTYPATH_API void ResolveLeaf(UFunction* InFunction) const;
 
 	/** 
 	 * Check whether a path is resolved. This means that it has a cached address, but may
 	 * resolve over an object boundary or a dynamic array.
 	 * @return true if the path is resolved
 	 */
-	bool IsResolved() const;
+	PROPERTYPATH_API bool IsResolved() const;
 
 	/** 
 	 * Check whether a path is fully resolved. This means that it has a cached address and
 	 * does not resolve over an object boundary or a dynamic array.
 	 * @return true if the path is fully resolved
 	 */
-	bool IsFullyResolved() const;
+	PROPERTYPATH_API bool IsFullyResolved() const;
 
 	/** Get the cached address for this property path, if any */
-	void* GetCachedAddress() const;
+	PROPERTYPATH_API void* GetCachedAddress() const;
 
 	/** Get the cached function for this property path, if any */
-	UFunction* GetCachedFunction() const;
+	PROPERTYPATH_API UFunction* GetCachedFunction() const;
 
 	/** Convert this property path to a FPropertyChangedEvent. Note that the path must be resolved. */
-	FPropertyChangedEvent ToPropertyChangedEvent(EPropertyChangeType::Type InChangeType) const;
+	PROPERTYPATH_API FPropertyChangedEvent ToPropertyChangedEvent(EPropertyChangeType::Type InChangeType) const;
 
 	/** Convert this property path to a FEditPropertyChain. Note that the path must be resolved. */
-	void ToEditPropertyChain(FEditPropertyChain& OutPropertyChain) const;
+	PROPERTYPATH_API void ToEditPropertyChain(FEditPropertyChain& OutPropertyChain) const;
 
 	/** Make a string representation of this property path */
-	FString ToString() const;
+	PROPERTYPATH_API FString ToString() const;
 
 	/** Compares this property path to a string */
-	bool operator==(const FString& Other) const;
+	PROPERTYPATH_API bool operator==(const FString& Other) const;
 
 	/** Compares this property path to a string */
-	bool Equals(const FString& Other) const;
+	PROPERTYPATH_API bool Equals(const FString& Other) const;
 
 	/** Get the cached container for this property path, for checking purposes */
-	void* GetCachedContainer() const;
+	PROPERTYPATH_API void* GetCachedContainer() const;
 
 	/** Set the cached container for this property path, for checking purposes */
-	void SetCachedContainer(void* InContainer) const;
+	PROPERTYPATH_API void SetCachedContainer(void* InContainer) const;
 
 	/** Trims this property path at the end */
-	void RemoveFromEnd(int32 InNumSegments = 1);
+	PROPERTYPATH_API void RemoveFromEnd(int32 InNumSegments = 1);
 
 	/** Trims this property path at the start */
-	void RemoveFromStart(int32 InNumSegments = 1);
+	PROPERTYPATH_API void RemoveFromStart(int32 InNumSegments = 1);
 
 	/** Returns FProperty if valid. This can be UFunction */
-	FProperty* GetFProperty() const;
+	PROPERTYPATH_API FProperty* GetFProperty() const;
 private:
 	/** Path segments for this path */
 	UPROPERTY()
diff --git a/Engine/Source/Runtime/RHI/Public/RHIResources.h b/Engine/Source/Runtime/RHI/Public/RHIResources.h
--- a/Engine/Source/Runtime/RHI/Public/RHIResources.h
+++ b/Engine/Source/Runtime/RHI/Public/RHIResources.h
@@ -4436,7 +4436,7 @@ struct FRHIRenderPassInfo
 #if !(UE_BUILD_SHIPPING || UE_BUILD_TEST)
 	RHI_API void Validate() const;
 #else
-	RHI_API void Validate() const {}
+	void Validate() const {}
 #endif
 	RHI_API void ConvertToRenderTargetsInfo(FRHISetRenderTargetsInfo& OutRTInfo) const;
 };
diff --git a/Engine/Source/Runtime/RHI/Public/RHIUtilities.h b/Engine/Source/Runtime/RHI/Public/RHIUtilities.h
--- a/Engine/Source/Runtime/RHI/Public/RHIUtilities.h
+++ b/Engine/Source/Runtime/RHI/Public/RHIUtilities.h
@@ -549,11 +549,11 @@ public:
 	RHI_API static void DumpResourceTransition(const FName& ResourceName, const ERHIAccess TransitionType);
 	
 private:
-	RHI_API static void DumpTransitionForResourceHandler();
+	static void DumpTransitionForResourceHandler();
 
-	RHI_API static TAutoConsoleVariable<FString> CVarDumpTransitionsForResource;
-	RHI_API static FAutoConsoleVariableSink CVarDumpTransitionsForResourceSink;
-	RHI_API static FName DumpTransitionForResource;
+	static TAutoConsoleVariable<FString> CVarDumpTransitionsForResource;
+	static FAutoConsoleVariableSink CVarDumpTransitionsForResourceSink;
+	static FName DumpTransitionForResource;
 };
 
 #if ENABLE_TRANSITION_DUMP
diff --git a/Engine/Source/Runtime/RHI/Public/TextureProfiler.h b/Engine/Source/Runtime/RHI/Public/TextureProfiler.h
--- a/Engine/Source/Runtime/RHI/Public/TextureProfiler.h
+++ b/Engine/Source/Runtime/RHI/Public/TextureProfiler.h
@@ -62,7 +62,7 @@ private:
 
 		RHI_API FTextureDetails& operator+=(const FTextureDetails& Other);
 		RHI_API FTextureDetails& operator-=(const FTextureDetails& Other);
-		RHI_API FTextureDetails& operator=(const FTextureDetails& Other) = default;
+		FTextureDetails& operator=(const FTextureDetails& Other) = default;
 
 		size_t Size = 0;
 		size_t PeakSize = 0;
diff --git a/Engine/Source/Runtime/RHICore/Public/RHICoreTransientResourceAllocator.h b/Engine/Source/Runtime/RHICore/Public/RHICoreTransientResourceAllocator.h
--- a/Engine/Source/Runtime/RHICore/Public/RHICoreTransientResourceAllocator.h
+++ b/Engine/Source/Runtime/RHICore/Public/RHICoreTransientResourceAllocator.h
@@ -78,7 +78,7 @@ struct FRHITransientResourceStats
 };
 
 /** Tracks all transient memory statistics for the current allocation cycle and reports results to various profilers. */
-struct RHICORE_API FRHITransientMemoryStats
+struct FRHITransientMemoryStats
 {
 public:
 	FRHITransientMemoryStats() = default;
@@ -123,7 +123,7 @@ public:
 		AliasedSize = AliasedSizeCurrent;
 	}
 
-	void Submit(uint64 TotalMemoryCapacity);
+	RHICORE_API void Submit(uint64 TotalMemoryCapacity);
 
 	bool HasDeallocations() const { return Textures.DeallocationCount > 0 || Buffers.DeallocationCount > 0; }
 
@@ -138,11 +138,11 @@ public:
 };
 
 /** Tracks resource allocations on the heap and adds overlap events to transient resources. */
-class RHICORE_API FRHITransientResourceOverlapTracker
+class FRHITransientResourceOverlapTracker
 {
 public:
-	void Track(FRHITransientResource* InResource, uint32 PageOffsetMin, uint32 PageOffsetMax);
-	void Reset();
+	RHICORE_API void Track(FRHITransientResource* InResource, uint32 PageOffsetMin, uint32 PageOffsetMax);
+	RHICORE_API void Reset();
 
 private:
 	struct FResourceRange
@@ -323,14 +323,14 @@ class FRHITransientHeapCache;
 class FRHITransientResourceHeapAllocator;
 
 /** First-fit allocator used for placing resources on a heap. */
-class RHICORE_API FRHITransientHeapAllocator
+class FRHITransientHeapAllocator
 {
 public:
-	FRHITransientHeapAllocator(uint64 Capacity, uint32 Alignment);
+	RHICORE_API FRHITransientHeapAllocator(uint64 Capacity, uint32 Alignment);
 
-	FRHITransientHeapAllocation Allocate(uint64 Size, uint32 Alignment);
+	RHICORE_API FRHITransientHeapAllocation Allocate(uint64 Size, uint32 Alignment);
 
-	void Deallocate(FRHITransientHeapAllocation Allocation);
+	RHICORE_API void Deallocate(FRHITransientHeapAllocation Allocation);
 
 	void SetGpuVirtualAddress(uint64 InGpuVirtualAddress)
 	{
@@ -409,9 +409,9 @@ private:
 		FRangeHandle FoundHandle = InvalidRangeHandle;
 	};
 
-	FFindResult FindFreeRange(uint64 Size, uint32 Alignment);
+	RHICORE_API FFindResult FindFreeRange(uint64 Size, uint32 Alignment);
 
-	void Validate();
+	RHICORE_API void Validate();
 
 	uint64 GpuVirtualAddress = 0;
 	uint64 Capacity = 0;
@@ -448,7 +448,7 @@ ENUM_CLASS_FLAGS(ERHITransientHeapFlags);
  *  byte offsets. Each heap additionally contains a cache of RHI transient resources, each with its own RHI
  *  resource and cache of RHI views. The lifetime of the resource cache is tied to the heap.
  */
-class RHICORE_API FRHITransientHeap
+class FRHITransientHeap
 {
 public:
 	struct FInitializer
@@ -502,7 +502,7 @@ public:
 
 	virtual ~FRHITransientHeap() = default;
 
-	FRHITransientTexture* CreateTexture(
+	RHICORE_API FRHITransientTexture* CreateTexture(
 		const FRHITextureCreateInfo& CreateInfo,
 		const TCHAR* DebugName,
 		uint32 PassIndex,
@@ -511,9 +511,9 @@ public:
 		uint32 TextureAlignment,
 		FCreateTextureFunction CreateTextureFunction);
 
-	void DeallocateMemory(FRHITransientTexture* Texture, uint32 PassIndex);
+	RHICORE_API void DeallocateMemory(FRHITransientTexture* Texture, uint32 PassIndex);
 
-	FRHITransientBuffer* CreateBuffer(
+	RHICORE_API FRHITransientBuffer* CreateBuffer(
 		const FRHIBufferCreateInfo& CreateInfo,
 		const TCHAR* DebugName,
 		uint32 PassIndex,
@@ -522,9 +522,9 @@ public:
 		uint32 BufferAlignment,
 		FCreateBufferFunction CreateBufferFunction);
 
-	void DeallocateMemory(FRHITransientBuffer* Buffer, uint32 PassIndex);
+	RHICORE_API void DeallocateMemory(FRHITransientBuffer* Buffer, uint32 PassIndex);
 
-	void Flush(uint64 CurrentAllocatorCycle, FRHITransientMemoryStats& OutMemoryStats, FRHITransientAllocationStats* OutAllocationStats);
+	RHICORE_API void Flush(uint64 CurrentAllocatorCycle, FRHITransientMemoryStats& OutMemoryStats, FRHITransientAllocationStats* OutAllocationStats);
 
 	const FInitializer& GetInitializer() const { return Initializer; }
 
@@ -554,8 +554,8 @@ protected:
 	}
 
 private:
-	void AllocateMemoryInternal(FRHITransientResource* Resource, const TCHAR* Name, uint32 PassIndex, uint64 CurrentAllocatorCycle, const FRHITransientHeapAllocation& Allocation);
-	void DeallocateMemoryInternal(FRHITransientResource* Resource, uint32 PassIndex);
+	RHICORE_API void AllocateMemoryInternal(FRHITransientResource* Resource, const TCHAR* Name, uint32 PassIndex, uint64 CurrentAllocatorCycle, const FRHITransientHeapAllocation& Allocation);
+	RHICORE_API void DeallocateMemoryInternal(FRHITransientResource* Resource, uint32 PassIndex);
 
 	FInitializer Initializer;
 	FRHITransientHeapAllocator Allocator;
@@ -583,13 +583,13 @@ private:
  *  are destroyed. Additionally, statistics are gathered automatically and reported to the 'rhitransientmemory' stats
  *  group.
  */
-class RHICORE_API FRHITransientHeapCache : public IRHITransientMemoryCache
+class FRHITransientHeapCache : public IRHITransientMemoryCache
 {
 public:
-	struct RHICORE_API FInitializer
+	struct FInitializer
 	{
 		// Creates a default initializer using common RHI CVars.
-		static FInitializer CreateDefault();
+		static RHICORE_API FInitializer CreateDefault();
 
 		static const uint32 kDefaultResourceCacheSize = 256;
 
@@ -619,13 +619,13 @@ public:
 		: Initializer(InInitializer)
 	{}
 
-	virtual ~FRHITransientHeapCache();
+	RHICORE_API virtual ~FRHITransientHeapCache();
 
-	FRHITransientHeap* Acquire(uint64 FirstAllocationSize, ERHITransientHeapFlags FirstAllocationHeapFlags);
+	RHICORE_API FRHITransientHeap* Acquire(uint64 FirstAllocationSize, ERHITransientHeapFlags FirstAllocationHeapFlags);
 
-	void Forfeit(TConstArrayView<FRHITransientHeap*> Heaps);
+	RHICORE_API void Forfeit(TConstArrayView<FRHITransientHeap*> Heaps);
 
-	void GarbageCollect() override;
+	RHICORE_API void GarbageCollect() override;
 
 	const FInitializer& GetInitializer() const { return Initializer; }
 
@@ -659,7 +659,7 @@ private:
 };
 
 /** A base class for implementing IRHITransientResourceAllocator for a virtual aliasing placed resource heap allocation strategy. */
-class RHICORE_API FRHITransientResourceHeapAllocator : public IRHITransientResourceAllocator
+class FRHITransientResourceHeapAllocator : public IRHITransientResourceAllocator
 {
 public:
 	FRHITransientResourceHeapAllocator(FRHITransientHeapCache& InHeapCache)
@@ -667,13 +667,13 @@ public:
 	{}
 
 	// Deallocates a texture from its parent heap. Provide the current platform fence value used to update the heap.
-	void DeallocateMemory(FRHITransientTexture* Texture, uint32 PassIndex) override;
+	RHICORE_API void DeallocateMemory(FRHITransientTexture* Texture, uint32 PassIndex) override;
 
 	// Deallocates a buffer from its parent heap. Provide the current platform fence value used to update the heap.
-	void DeallocateMemory(FRHITransientBuffer* Buffer, uint32 PassIndex) override;
+	RHICORE_API void DeallocateMemory(FRHITransientBuffer* Buffer, uint32 PassIndex) override;
 
 	// Called to flush any active allocations prior to rendering.
-	void Flush(FRHICommandListImmediate& RHICmdList, FRHITransientAllocationStats* OutAllocationStats) override;
+	RHICORE_API void Flush(FRHICommandListImmediate& RHICmdList, FRHITransientAllocationStats* OutAllocationStats) override;
 
 	// Returns the array of heaps used by this allocator, including the required commit size for each.
 	inline TConstArrayView<FRHITransientHeap*> GetHeaps() const { return Heaps; }
@@ -686,7 +686,7 @@ protected:
 	 *  texture create info and passed in, along with a platform-specific texture creation function if no cached resource
 	 *  if found.
 	 */
-	FRHITransientTexture* CreateTextureInternal(
+	RHICORE_API FRHITransientTexture* CreateTextureInternal(
 		const FRHITextureCreateInfo& CreateInfo,
 		const TCHAR* DebugName,
 		uint32 PassIndex,
@@ -699,7 +699,7 @@ protected:
 	 *  buffer create info and passed in, along with a platform-specific buffer creation function if no cached resource
 	 *  if found.
 	 */
-	FRHITransientBuffer* CreateBufferInternal(
+	RHICORE_API FRHITransientBuffer* CreateBufferInternal(
 		const FRHIBufferCreateInfo& CreateInfo,
 		const TCHAR* DebugName,
 		uint32 PassIndex,
@@ -723,7 +723,7 @@ class FRHITransientPagePoolCache;
 class FRHITransientResourcePageAllocator;
 
 /** Allocates page spans for a resource. */
-class RHICORE_API FRHITransientPageSpanAllocator
+class FRHITransientPageSpanAllocator
 {
 public:
 	FRHITransientPageSpanAllocator(uint32 InPageCount, uint32 InPageSize)
@@ -734,11 +734,11 @@ public:
 		Init();
 	}
 
-	void Reset();
+	RHICORE_API void Reset();
 
-	bool Allocate(uint32 PageCount, uint32& NumPagesAllocated, uint32& SpanIndex);
+	RHICORE_API bool Allocate(uint32 PageCount, uint32& NumPagesAllocated, uint32& SpanIndex);
 
-	void Deallocate(uint32 SpanIndex);
+	RHICORE_API void Deallocate(uint32 SpanIndex);
 
 	template <typename SpanArrayType>
 	void GetSpanArray(uint32 SpanIndex, SpanArrayType& OutPageSpans) const
@@ -759,7 +759,7 @@ public:
 
 	uint64 GetFreeSize() const { return FreePageCount * PageSize; }
 
-	uint32 GetAllocationPageCount(uint32 SpanIndex) const;
+	RHICORE_API uint32 GetAllocationPageCount(uint32 SpanIndex) const;
 
 	uint32 GetMaxSpanCount() const { return MaxSpanCount; }
 
@@ -785,22 +785,22 @@ private:
 		bool bAllocated = false;
 	};
 
-	void Init();
+	RHICORE_API void Init();
 
 	// Splits a span into two, so that the original span has PageCount pages and the new span contains the remaining ones
-	void SplitSpan(uint32 SpanIndex, uint32 PageCount);
+	RHICORE_API void SplitSpan(uint32 SpanIndex, uint32 PageCount);
 
 	// Merges two spans. They must be adjacent and in the same list
-	void MergeSpans(uint32 SpanIndex0, uint32 SpanIndex1, const bool bKeepSpan1);
+	RHICORE_API void MergeSpans(uint32 SpanIndex0, uint32 SpanIndex1, const bool bKeepSpan1);
 
 	// Inserts a span after an existing span. The span to insert must be unlinked
-	void InsertAfter(uint32 InsertPosition, uint32 InsertSpanIndex);
+	RHICORE_API void InsertAfter(uint32 InsertPosition, uint32 InsertSpanIndex);
 
 	// Inserts a span after an existing span. The span to insert must be unlinked
-	void InsertBefore(uint32 InsertPosition, uint32 InsertSpanIndex);
+	RHICORE_API void InsertBefore(uint32 InsertPosition, uint32 InsertSpanIndex);
 
 	// Removes a span from its list, reconnecting neighbouring list elements
-	void Unlink(uint32 SpanIndex);
+	RHICORE_API void Unlink(uint32 SpanIndex);
 
 	// Allocates an unused span from the pool
 	int AllocSpan()
@@ -821,9 +821,9 @@ private:
 	}
 
 	// Merges a span with existing neighbours in the free list if they exist
-	bool MergeFreeSpanIfPossible(uint32 SpanIndex);
+	RHICORE_API bool MergeFreeSpanIfPossible(uint32 SpanIndex);
 
-	void Validate();
+	RHICORE_API void Validate();
 
 	int GetFirstSpanIndex() const
 	{
@@ -846,7 +846,7 @@ private:
 
 ///////////////////////////////////////////////////////////////////////////////////////////////////
 
-class RHICORE_API FRHITransientPagePool
+class FRHITransientPagePool
 {
 public:
 	struct FInitializer
@@ -895,14 +895,14 @@ public:
 		uint32 PagesMapped = 0;
 	};
 
-	void Allocate(FAllocationContext& AllocationContext);
+	RHICORE_API void Allocate(FAllocationContext& AllocationContext);
 
 	void Deallocate(uint32 SpanIndex)
 	{
 		Allocator.Deallocate(SpanIndex);
 	}
 
-	void Flush(FRHICommandListImmediate& RHICmdList);
+	RHICORE_API void Flush(FRHICommandListImmediate& RHICmdList);
 
 	uint64 GetLastUsedGarbageCollectCycle() const { return LastUsedGarbageCollectCycle; }
 
@@ -971,7 +971,7 @@ private:
 
 ///////////////////////////////////////////////////////////////////////////////////////////////////
 
-class RHICORE_API FRHITransientPagePoolCache : public IRHITransientMemoryCache
+class FRHITransientPagePoolCache : public IRHITransientMemoryCache
 {
 public:
 	struct FInitializer
@@ -1006,18 +1006,18 @@ public:
 
 	FRHITransientPagePoolCache(const FRHITransientPagePoolCache&) = delete;
 
-	virtual ~FRHITransientPagePoolCache();
+	RHICORE_API virtual ~FRHITransientPagePoolCache();
 
 	// Called by the transient allocator to acquire a page pool from the cache.
-	FRHITransientPagePool* Acquire();
+	RHICORE_API FRHITransientPagePool* Acquire();
 
 	// Called by the transient allocator to return the fast page pool if it exists.
-	FRHITransientPagePool* GetFastPagePool();
+	RHICORE_API FRHITransientPagePool* GetFastPagePool();
 
 	// Called by the transient allocator to forfeit all acquired heaps back to the cache.
-	void Forfeit(TConstArrayView<FRHITransientPagePool*> PagePools);
+	RHICORE_API void Forfeit(TConstArrayView<FRHITransientPagePool*> PagePools);
 
-	void GarbageCollect() override;
+	RHICORE_API void GarbageCollect() override;
 
 	const FInitializer Initializer;
 
@@ -1047,7 +1047,7 @@ private:
 
 ///////////////////////////////////////////////////////////////////////////////////////////////////
 
-class RHICORE_API FRHITransientResourcePageAllocator : public IRHITransientResourceAllocator
+class FRHITransientResourcePageAllocator : public IRHITransientResourceAllocator
 {
 public:
 	FRHITransientResourcePageAllocator(FRHITransientPagePoolCache& InPagePoolCache)
@@ -1059,11 +1059,11 @@ public:
 		FastPagePool = PagePoolCache.GetFastPagePool();
 	}
 
-	FRHITransientTexture* CreateTexture(const FRHITextureCreateInfo& CreateInfo, const TCHAR* DebugName, uint32 PassIndex) override;
-	FRHITransientBuffer* CreateBuffer(const FRHIBufferCreateInfo& CreateInfo, const TCHAR* DebugName, uint32 PassIndex) override;
-	void DeallocateMemory(FRHITransientTexture* Texture, uint32 PassIndex) override;
-	void DeallocateMemory(FRHITransientBuffer* Buffer, uint32 PassIndex) override;
-	void Flush(FRHICommandListImmediate& RHICmdList, FRHITransientAllocationStats* OutAllocationStats) override;
+	RHICORE_API FRHITransientTexture* CreateTexture(const FRHITextureCreateInfo& CreateInfo, const TCHAR* DebugName, uint32 PassIndex) override;
+	RHICORE_API FRHITransientBuffer* CreateBuffer(const FRHIBufferCreateInfo& CreateInfo, const TCHAR* DebugName, uint32 PassIndex) override;
+	RHICORE_API void DeallocateMemory(FRHITransientTexture* Texture, uint32 PassIndex) override;
+	RHICORE_API void DeallocateMemory(FRHITransientBuffer* Buffer, uint32 PassIndex) override;
+	RHICORE_API void Flush(FRHICommandListImmediate& RHICmdList, FRHITransientAllocationStats* OutAllocationStats) override;
 
 	uint32 GetPageSize() const { return PageSize; }
 
@@ -1076,8 +1076,8 @@ private:
 	static constexpr uint64 KB = 1024;
 	static constexpr uint64 MB = 1024 * KB;
 
-	void AllocateMemoryInternal(FRHITransientResource* Resource, const TCHAR* DebugName, uint32 PassIndex, bool bFastPoolRequested, float FastPoolPercentageRequested);
-	void DeallocateMemoryInternal(FRHITransientResource* Resource, uint32 PassIndex);
+	RHICORE_API void AllocateMemoryInternal(FRHITransientResource* Resource, const TCHAR* DebugName, uint32 PassIndex, bool bFastPoolRequested, float FastPoolPercentageRequested);
+	RHICORE_API void DeallocateMemoryInternal(FRHITransientResource* Resource, uint32 PassIndex);
 
 	//////////////////////////////////////////////////////////////////////////
 	//! Platform API
@@ -1126,4 +1126,4 @@ private:
 	uint32 PageSpanCount = 0;
 
 	IF_RHICORE_TRANSIENT_ALLOCATOR_DEBUG(TSet<FRHITransientResource*> ActiveResources);
-};
\ No newline at end of file
+};
