FRHIOffsetHeapDescriptorAllocator::Allocate wasn't handling invalid descriptors from FRHIHeapDescriptorAllocator::Allocate

diff --git a/Engine/Source/Runtime/RHICore/Private/RHIDescriptorAllocator.cpp b/Engine/Source/Runtime/RHICore/Private/RHIDescriptorAllocator.cpp
--- a/Engine/Source/Runtime/RHICore/Private/RHIDescriptorAllocator.cpp
+++ b/Engine/Source/Runtime/RHICore/Private/RHIDescriptorAllocator.cpp
@@ -242,7 +242,11 @@ FRHIOffsetHeapDescriptorAllocator::FRHIOffsetHeapDescriptorAllocator(ERHIDescrip
 FRHIDescriptorHandle FRHIOffsetHeapDescriptorAllocator::Allocate()
 {
 	const FRHIDescriptorHandle AlocatorHandle = FRHIHeapDescriptorAllocator::Allocate();
-	return FRHIDescriptorHandle(AlocatorHandle.GetType(), AlocatorHandle.GetIndex() + HeapOffset);
+	if (AlocatorHandle.IsValid())
+	{
+		return FRHIDescriptorHandle(AlocatorHandle.GetType(), AlocatorHandle.GetIndex() + HeapOffset);
+	}
+	return FRHIDescriptorHandle();
 }
 
 void FRHIOffsetHeapDescriptorAllocator::Free(const FRHIDescriptorHandle InHandle)
