Adding variant of UE::RHICore::ApplyStaticUniformBuffers that takes a lambda instead of a Context.

diff --git a/Engine/Source/Runtime/RHICore/Public/RHICoreShader.h b/Engine/Source/Runtime/RHICore/Public/RHICoreShader.h
--- a/Engine/Source/Runtime/RHICore/Public/RHICoreShader.h
+++ b/Engine/Source/Runtime/RHICore/Public/RHICoreShader.h
@@ -42,13 +42,14 @@ inline void InitStaticUniformBufferSlots(TArray<FUniformBufferStaticSlot>& Stati
 	}
 }
 
-template <typename TRHIContext, typename TRHIShader>
+template <typename TApplyFunction>
 void ApplyStaticUniformBuffers(
-	TRHIContext* CommandContext,
-	TRHIShader* Shader,
+	FRHIShader* Shader,
 	const TArray<FUniformBufferStaticSlot>& Slots,
 	const TArray<uint32>& LayoutHashes,
-	const TArray<FRHIUniformBuffer*>& UniformBuffers)
+	const TArray<FRHIUniformBuffer*>& UniformBuffers,
+	TApplyFunction&& ApplyFunction
+)
 {
 	checkf(LayoutHashes.Num() == Slots.Num(), TEXT("Shader %s, LayoutHashes %d, Slots %d"),
 		Shader->GetShaderName(), LayoutHashes.Num(), Slots.Num());
@@ -64,12 +65,27 @@ void ApplyStaticUniformBuffers(
 
 			if (Buffer)
 			{
-				CommandContext->RHISetShaderUniformBuffer(Shader, BufferIndex, Buffer);
+				ApplyFunction(BufferIndex, Buffer);
 			}
 		}
 	}
 }
 
+template <typename TRHIContext, typename TRHIShader>
+void ApplyStaticUniformBuffers(
+	TRHIContext* CommandContext,
+	TRHIShader* Shader,
+	const TArray<FUniformBufferStaticSlot>& Slots,
+	const TArray<uint32>& LayoutHashes,
+	const TArray<FRHIUniformBuffer*>& UniformBuffers)
+{
+	ApplyStaticUniformBuffers(Shader, Slots, LayoutHashes, UniformBuffers,
+		[CommandContext, Shader](int32 BufferIndex, FRHIUniformBuffer* Buffer)
+		{
+			CommandContext->RHISetShaderUniformBuffer(Shader, BufferIndex, Buffer);
+		});
+}
+
 template <typename TResourceType> struct TResourceTypeStr {};
 template<> struct TResourceTypeStr<FRHISamplerState       > { static constexpr TCHAR String[] = TEXT("Sampler State"); };
 template<> struct TResourceTypeStr<FRHITexture            > { static constexpr TCHAR String[] = TEXT("Texture"); };
