Fixing up SRV support for PC DX12 P010 format: plane selection code was missing

diff --git a/Engine/Source/Runtime/RHICore/Internal/DXGIUtilities.h b/Engine/Source/Runtime/RHICore/Internal/DXGIUtilities.h
--- a/Engine/Source/Runtime/RHICore/Internal/DXGIUtilities.h
+++ b/Engine/Source/Runtime/RHICore/Internal/DXGIUtilities.h
@@ -230,6 +230,14 @@ namespace UE::DXGIUtilities
 			case DXGI_FORMAT_R8G8_UNORM: return 1;
 			}
 			break;
+
+		case DXGI_FORMAT_P010:
+			switch (ViewFormat)
+			{
+			case DXGI_FORMAT_R16_UNORM: return 0;
+			case DXGI_FORMAT_R16G16_UNORM: return 1;
+			}
+			break;
 		}
 
 		return 0;
