Fix C6295
#rnx

diff --git a/Engine/Source/Runtime/RHICore/Private/RHICoreTransientResourceAllocator.cpp b/Engine/Source/Runtime/RHICore/Private/RHICoreTransientResourceAllocator.cpp
--- a/Engine/Source/Runtime/RHICore/Private/RHICoreTransientResourceAllocator.cpp
+++ b/Engine/Source/Runtime/RHICore/Private/RHICoreTransientResourceAllocator.cpp
@@ -1185,8 +1185,8 @@ void FRHITransientPageSpanAllocator::Validate()
 
 	// Count free pages
 	uint32 FreeCount = 0;
-	int32 PrevIndex = FreeSpanListHeadIndex;
-	for (int32 Index = GetFirstSpanIndex(); Index != InvalidIndex; Index = PageSpans[Index].NextSpanIndex)
+	uint32 PrevIndex = FreeSpanListHeadIndex;
+	for (uint32 Index = GetFirstSpanIndex(); PageSpans.IsValidIndex(Index); Index = PageSpans[Index].NextSpanIndex)
 	{
 		FPageSpan& Span = PageSpans[Index];
 		check(Span.PrevSpanIndex == PrevIndex);
diff --git a/Engine/Source/Runtime/RHICore/Public/RHICoreTransientResourceAllocator.h b/Engine/Source/Runtime/RHICore/Public/RHICoreTransientResourceAllocator.h
--- a/Engine/Source/Runtime/RHICore/Public/RHICoreTransientResourceAllocator.h
+++ b/Engine/Source/Runtime/RHICore/Public/RHICoreTransientResourceAllocator.h
@@ -825,7 +825,7 @@ private:
 
 	RHICORE_API void Validate();
 
-	int GetFirstSpanIndex() const
+	uint32 GetFirstSpanIndex() const
 	{
 		return PageSpans[FreeSpanListHeadIndex].NextSpanIndex;
 	}
