Ran IWYU on RHICore

diff --git a/Engine/Source/Runtime/RHICore/Internal/RHIShaderParametersShared.h b/Engine/Source/Runtime/RHICore/Internal/RHIShaderParametersShared.h
--- a/Engine/Source/Runtime/RHICore/Internal/RHIShaderParametersShared.h
+++ b/Engine/Source/Runtime/RHICore/Internal/RHIShaderParametersShared.h
@@ -2,7 +2,7 @@
 
 #pragma once
 
-#include "RHI.h"
+#include "RHIShaderParameters.h"
 
 namespace UE::RHICore
 {
diff --git a/Engine/Source/Runtime/RHICore/Internal/RHIUniformBufferDataShared.h b/Engine/Source/Runtime/RHICore/Internal/RHIUniformBufferDataShared.h
--- a/Engine/Source/Runtime/RHICore/Internal/RHIUniformBufferDataShared.h
+++ b/Engine/Source/Runtime/RHICore/Internal/RHIUniformBufferDataShared.h
@@ -2,7 +2,6 @@
 
 #pragma once
 
-#include "RHI.h"
 #include "RenderGraphResources.h"
 
 namespace UE::RHICore
diff --git a/Engine/Source/Runtime/RHICore/Private/RHICoreTransientResourceAllocator.cpp b/Engine/Source/Runtime/RHICore/Private/RHICoreTransientResourceAllocator.cpp
--- a/Engine/Source/Runtime/RHICore/Private/RHICoreTransientResourceAllocator.cpp
+++ b/Engine/Source/Runtime/RHICore/Private/RHICoreTransientResourceAllocator.cpp
@@ -7,7 +7,6 @@
 #include "ProfilingDebugging/CountersTrace.h"
 #include "ProfilingDebugging/CsvProfiler.h"
 #include "RHICommandList.h"
-#include "RHICore.h"
 
 static int32 GRHITransientAllocatorMinimumHeapSize = 128;
 static FAutoConsoleVariableRef CVarRHITransientAllocatorMinimumHeapSize(
diff --git a/Engine/Source/Runtime/RHICore/Private/RHIDescriptorAllocator.cpp b/Engine/Source/Runtime/RHICore/Private/RHIDescriptorAllocator.cpp
--- a/Engine/Source/Runtime/RHICore/Private/RHIDescriptorAllocator.cpp
+++ b/Engine/Source/Runtime/RHICore/Private/RHIDescriptorAllocator.cpp
@@ -2,6 +2,7 @@
 
 #include "RHIDescriptorAllocator.h"
 #include "Misc/ScopeLock.h"
+#include "RHIDefinitions.h"
 
 struct FRHIDescriptorAllocatorRange
 {
diff --git a/Engine/Source/Runtime/RHICore/Public/BoundShaderStateHistory.h b/Engine/Source/Runtime/RHICore/Public/BoundShaderStateHistory.h
--- a/Engine/Source/Runtime/RHICore/Public/BoundShaderStateHistory.h
+++ b/Engine/Source/Runtime/RHICore/Public/BoundShaderStateHistory.h
@@ -2,8 +2,8 @@
 
 #pragma once
 
+#include "RHIGlobals.h"
 #include "RenderResource.h"
-#include "RHI.h"
 
 /**
  * A list of the most recently used bound shader states.
@@ -60,3 +60,7 @@ private:
 	uint32 NextBoundShaderStateIndex = 0;
 	FCriticalSection BoundShaderStateHistoryLock;
 };
+
+#if UE_ENABLE_INCLUDE_ORDER_DEPRECATED_IN_5_4
+#include "RHI.h"
+#endif
diff --git a/Engine/Source/Runtime/RHICore/Public/RHICore.h b/Engine/Source/Runtime/RHICore/Public/RHICore.h
--- a/Engine/Source/Runtime/RHICore/Public/RHICore.h
+++ b/Engine/Source/Runtime/RHICore/Public/RHICore.h
@@ -3,9 +3,7 @@
 #pragma once
 
 #include "Logging/LogMacros.h"
-#include "RHI.h"
 #include "RHIResources.h"
-#include "RHIContext.h"
 
 RHICORE_API DECLARE_LOG_CATEGORY_EXTERN(LogRHICore, Log, VeryVerbose);
 
@@ -37,4 +35,9 @@ struct FResolveTextureInfo
 
 RHICORE_API void ResolveRenderPassTargets(const FRHIRenderPassInfo& Info, TFunction<void(FResolveTextureInfo)> Function);
 
-} //! UE::RHICore
\ No newline at end of file
+} //! UE::RHICore
+
+#if UE_ENABLE_INCLUDE_ORDER_DEPRECATED_IN_5_4
+#include "RHI.h"
+#include "RHIContext.h"
+#endif
diff --git a/Engine/Source/Runtime/RHICore/Public/RHICoreShader.h b/Engine/Source/Runtime/RHICore/Public/RHICoreShader.h
--- a/Engine/Source/Runtime/RHICore/Public/RHICoreShader.h
+++ b/Engine/Source/Runtime/RHICore/Public/RHICoreShader.h
@@ -2,8 +2,12 @@
 
 #pragma once
 
+#include "Misc/App.h"
 #include "RHICore.h"
 #include "ShaderCore.h"
+#include "ShaderParameterMacros.h"
+
+struct FRHIShaderBundleDispatch;
 
 namespace UE
 {
@@ -261,4 +265,4 @@ void SetResourcesFromTables(TBinder&& Binder, FRHIShader const& Shader, FShaderR
 }
 
 } //! RHICore
-} //! UE
\ No newline at end of file
+} //! UE
diff --git a/Engine/Source/Runtime/RHICore/Public/RHIDescriptorAllocator.h b/Engine/Source/Runtime/RHICore/Public/RHIDescriptorAllocator.h
--- a/Engine/Source/Runtime/RHICore/Public/RHIDescriptorAllocator.h
+++ b/Engine/Source/Runtime/RHICore/Public/RHIDescriptorAllocator.h
@@ -3,9 +3,11 @@
 
 #include "Containers/Array.h"
 #include "HAL/CriticalSection.h"
-#include "RHIDefinitions.h"
 #include "Stats/Stats.h"
 
+enum class ERHIDescriptorHeapType : uint8;
+struct FRHIDescriptorHandle;
+
 struct FRHIDescriptorAllocatorRange;
 
 class FRHIDescriptorAllocator
@@ -95,3 +97,7 @@ private:
 	// Offset from start of heap we belong to
 	uint32 HeapOffset;
 };
+
+#if UE_ENABLE_INCLUDE_ORDER_DEPRECATED_IN_5_4
+#include "RHIDefinitions.h"
+#endif
diff --git a/Engine/Source/Runtime/RHICore/Public/RHIPoolAllocator.h b/Engine/Source/Runtime/RHICore/Public/RHIPoolAllocator.h
--- a/Engine/Source/Runtime/RHICore/Public/RHIPoolAllocator.h
+++ b/Engine/Source/Runtime/RHICore/Public/RHIPoolAllocator.h
@@ -1,8 +1,11 @@
 // Copyright Epic Games, Inc. All Rights Reserved.
 #pragma once
 
-#include "RHI.h"
-#include "RHIResources.h"
+#include "Containers/Array.h"
+#include "HAL/CriticalSection.h"
+
+class FRHICommandListBase;
+struct FScriptContainerElement;
 
 // Pre declares
 class FRHIPoolAllocator;
@@ -266,3 +269,8 @@ protected:
 	uint32 TotalAllocatedBlocks;
 };
 
+
+#if UE_ENABLE_INCLUDE_ORDER_DEPRECATED_IN_5_4
+#include "RHI.h"
+#include "RHIResources.h"
+#endif
