Fixed modular build link error (RHICORE_API was missing in LLM declare / define tag of RHITransientResources)

#rnx

diff --git a/Engine/Source/Runtime/RHICore/Private/RHICoreTransientResourceAllocator.cpp b/Engine/Source/Runtime/RHICore/Private/RHICoreTransientResourceAllocator.cpp
--- a/Engine/Source/Runtime/RHICore/Private/RHICoreTransientResourceAllocator.cpp
+++ b/Engine/Source/Runtime/RHICore/Private/RHICoreTransientResourceAllocator.cpp
@@ -70,7 +70,7 @@ DECLARE_DWORD_ACCUMULATOR_STAT(TEXT("Textures"), STAT_RHITransientTextures, STAT
 DECLARE_DWORD_ACCUMULATOR_STAT(TEXT("Buffers"), STAT_RHITransientBuffers, STATGROUP_RHITransientMemory);
 
 DECLARE_LLM_MEMORY_STAT(TEXT("RHI Transient Resources"), STAT_RHITransientResourcesLLM, STATGROUP_LLMFULL);
-LLM_DEFINE_TAG(RHITransientResources, NAME_None, NAME_None, GET_STATFNAME(STAT_RHITransientResourcesLLM), GET_STATFNAME(STAT_EngineSummaryLLM));
+RHICORE_API LLM_DEFINE_TAG(RHITransientResources, NAME_None, NAME_None, GET_STATFNAME(STAT_RHITransientResourcesLLM), GET_STATFNAME(STAT_EngineSummaryLLM));
 
 //////////////////////////////////////////////////////////////////////////
 
