Transient allocator memory improvements
 - Async compute resources are able to alias with each other on the async compute pipe.
        - Added r.RDG.AsyncComputeTransientAliasing and GRHIGlobals.SupportsAsyncComputeTransientAliasing to control at runtime and per-platform, respectively.
 - Offload D3D12 create placed resource calls to tasks.
        - Placed resource creation in D3D12 is very expensive relative to other platforms.
 - Fixed tracking of transient memory on platforms that support virtual pages.
 - Added more debug information to RDG insights.
 - Refactored D3D12 barrier implementation to support acquire / discard operations on async compute.

diff --git a/Engine/Plugins/RenderGraphInsights/Source/RenderGraphInsights/Private/RenderGraphProvider.cpp b/Engine/Plugins/RenderGraphInsights/Source/RenderGraphInsights/Private/RenderGraphProvider.cpp
--- a/Engine/Plugins/RenderGraphInsights/Source/RenderGraphInsights/Private/RenderGraphProvider.cpp
+++ b/Engine/Plugins/RenderGraphInsights/Source/RenderGraphInsights/Private/RenderGraphProvider.cpp
@@ -124,6 +124,9 @@ FResourcePacket::FResourcePacket(const UE::Trace::IAnalyzer::FOnEventContext& Co
 		TransientAllocations[LocalIndex].MemoryRangeIndex = TransientAllocationMemoryRanges[LocalIndex];
 	}
 
+	TransientAcquirePass = Context.EventData.GetValue<FRDGPassHandle>("TransientAcquirePass");
+	TransientDiscardPass = Context.EventData.GetValue<FRDGPassHandle>("TransientDiscardPass");
+
 	if (Passes.Num())
 	{
 		FirstPass = Passes[0];
diff --git a/Engine/Plugins/RenderGraphInsights/Source/RenderGraphInsights/Private/RenderGraphProvider.h b/Engine/Plugins/RenderGraphInsights/Source/RenderGraphInsights/Private/RenderGraphProvider.h
--- a/Engine/Plugins/RenderGraphInsights/Source/RenderGraphInsights/Private/RenderGraphProvider.h
+++ b/Engine/Plugins/RenderGraphInsights/Source/RenderGraphInsights/Private/RenderGraphProvider.h
@@ -84,6 +84,9 @@ public:
 
 	mutable TArray<uint32, TInlineAllocator<1>> VisibleItems;
 
+	FRDGPassHandle TransientAcquirePass;
+	FRDGPassHandle TransientDiscardPass;
+
 	bool bExternal{};
 	bool bExtracted{};
 	bool bCulled{};
diff --git a/Engine/Plugins/RenderGraphInsights/Source/RenderGraphInsights/Private/RenderGraphTrack.cpp b/Engine/Plugins/RenderGraphInsights/Source/RenderGraphInsights/Private/RenderGraphTrack.cpp
--- a/Engine/Plugins/RenderGraphInsights/Source/RenderGraphInsights/Private/RenderGraphTrack.cpp
+++ b/Engine/Plugins/RenderGraphInsights/Source/RenderGraphInsights/Private/RenderGraphTrack.cpp
@@ -825,7 +825,6 @@ void FRenderGraphTrack::Draw(const ITimingTrackDrawContext& Context) const
 		{
 			const FVisiblePass& AsyncComputeVisiblePass = VisibleGraph.Passes[VisibleIndex];
 			const FPassPacket& AsyncComputePass = AsyncComputeVisiblePass.GetPacket();
-			const FBoundingBox AsyncComputeVisiblePassBox = AsyncComputeVisiblePass.GetBoundingBox(Layout);
 
 			float TintAlpha = 0.25f;
 
@@ -1508,12 +1507,12 @@ void FRenderGraphTrack::BuildFilteredDrawState(FRenderGraphTrackDrawStateBuilder
 			VisibleItems.Add(&GraphicsVisiblePass);
 		};
 
-		if (Pass.bAsyncComputeBegin)
+		if (Pass.GraphicsForkPass.IsValid())
 		{
 			AddFencePassEvent(Pass.GraphicsForkPass);
 		}
 
-		if (Pass.bAsyncComputeEnd)
+		if (Pass.GraphicsJoinPass.IsValid())
 		{
 			AddFencePassEvent(Pass.GraphicsJoinPass);
 		}
@@ -2201,6 +2200,16 @@ void FRenderGraphTrack::InitTooltip(FTooltipDrawState& Tooltip, const ITimingEve
 			{
 				Tooltip.AddTextLine(TEXT("Transient"), FLinearColor::Red);
 				Tooltip.AddNameValueTextLine(TEXT("Transient Cache Hit:"), Resource.bTransientCacheHit ? TEXT("Yes") : TEXT("No"));
+
+				if (Resource.TransientAcquirePass.IsValid())
+				{
+					Tooltip.AddNameValueTextLine(TEXT("Transient Acquire Pass:"), FString::Printf(TEXT("%u"), Resource.TransientAcquirePass.GetIndex()));
+				}
+				
+				if (Resource.TransientDiscardPass.IsValid())
+				{
+					Tooltip.AddNameValueTextLine(TEXT("Transient Discard Pass:"), FString::Printf(TEXT("%u"), Resource.TransientDiscardPass.GetIndex()));
+				}
 			}
 
 			if (!Resource.TransientAllocations.IsEmpty())
@@ -2273,6 +2282,16 @@ void FRenderGraphTrack::InitTooltip(FTooltipDrawState& Tooltip, const ITimingEve
 			{
 				Tooltip.AddTextLine(TEXT("Merged RenderPass"), FLinearColor::Red);
 			}
+
+			if (Pass.GraphicsForkPass.IsValid())
+			{
+				Tooltip.AddNameValueTextLine(TEXT("GraphicsForkPass:"), FString::Printf(TEXT("%d"), Pass.GraphicsForkPass.GetIndex()));
+			}
+			
+			if (Pass.GraphicsJoinPass.IsValid())
+			{
+				Tooltip.AddNameValueTextLine(TEXT("GraphicsJoinPass:"), FString::Printf(TEXT("%d"), Pass.GraphicsJoinPass.GetIndex()));
+			}
 		}
 		else if (InTooltipEvent.Is<FVisibleTextureEvent>())
 		{
diff --git a/Engine/Source/Runtime/D3D12RHI/Private/D3D12Adapter.cpp b/Engine/Source/Runtime/D3D12RHI/Private/D3D12Adapter.cpp
--- a/Engine/Source/Runtime/D3D12RHI/Private/D3D12Adapter.cpp
+++ b/Engine/Source/Runtime/D3D12RHI/Private/D3D12Adapter.cpp
@@ -1115,6 +1115,8 @@ void FD3D12Adapter::InitializeDevices()
 #endif
 		}
 
+		GRHIGlobals.SupportsAsyncComputeTransientAliasing = GSupportsEfficientAsyncCompute;
+
 #if PLATFORM_WINDOWS
 		D3D12_FEATURE_DATA_D3D12_OPTIONS2 D3D12Caps2 = {};
 		if (FAILED(RootDevice->CheckFeatureSupport(D3D12_FEATURE_D3D12_OPTIONS2, &D3D12Caps2, sizeof(D3D12Caps2))))
diff --git a/Engine/Source/Runtime/D3D12RHI/Private/D3D12Commands.cpp b/Engine/Source/Runtime/D3D12RHI/Private/D3D12Commands.cpp
--- a/Engine/Source/Runtime/D3D12RHI/Private/D3D12Commands.cpp
+++ b/Engine/Source/Runtime/D3D12RHI/Private/D3D12Commands.cpp
@@ -8,13 +8,6 @@ D3D12Commands.cpp: D3D RHI commands implementation.
 #include "ProfilingDebugging/AssetMetadataTrace.h"
 #include "ProfilingDebugging/RealtimeGPUProfiler.h"
 
-static int32 GD3D12TransientAllocatorFullAliasingBarrier = 0;
-static FAutoConsoleVariableRef CVarD3D12TransientAllocatorFullAliasingBarrier(
-	TEXT("d3d12.TransientAllocator.FullAliasingBarrier"),
-	GD3D12TransientAllocatorFullAliasingBarrier,
-	TEXT("Inserts a full aliasing barrier on an transient acquire operation. Useful to debug if an aliasing barrier is missing."),
-	ECVF_RenderThreadSafe);
-
 static int32 GD3D12AllowDiscardResources = 1;
 static FAutoConsoleVariableRef CVarD3D12AllowDiscardResources(
 	TEXT("d3d12.AllowDiscardResources"),
@@ -266,7 +259,7 @@ void ProcessResource(FD3D12CommandContext& Context, const FRHITransitionInfo& In
 static bool ProcessTransitionDuringBegin(const FD3D12TransitionData* Data)
 {
 	// Pipe changes which are not ending with graphics or targeting all pipelines are handle during begin
-	return ((Data->SrcPipelines != Data->DstPipelines && Data->DstPipelines != ERHIPipeline::Graphics) || EnumHasAllFlags(Data->DstPipelines, ERHIPipeline::All));
+	return !EnumHasAllFlags(Data->SrcPipelines, ERHIPipeline::All) && ((Data->SrcPipelines != Data->DstPipelines && Data->DstPipelines != ERHIPipeline::Graphics) || EnumHasAllFlags(Data->DstPipelines, ERHIPipeline::All));
 }
 
 struct FD3D12DiscardResource
@@ -327,12 +320,10 @@ void FD3D12CommandContext::HandleResourceDiscardTransitions(
 				return;
 			}
 
-			// Get the initial state to force a 'nop' transition so the internal command list resource tracking has the correct state
-			// already to make sure it's not added to the pending transition list to be kicked before this command list (resource is not valid then yet)
-			D3D12_RESOURCE_STATES InitialState = GetInitialResourceState(Resource->GetDesc());
+			D3D12_RESOURCE_STATES InitialState = GetInitialResourceState(Resource->GetDesc(), QueueType);
 			if (Info.IsWholeResource() || Resource->GetSubresourceCount() == 1)
 			{
-				TransitionResource(Resource, InitialState, InitialState, D3D12_RESOURCE_BARRIER_ALL_SUBRESOURCES);
+				TransitionResource(Resource, D3D12_RESOURCE_STATE_TBD, InitialState, D3D12_RESOURCE_BARRIER_ALL_SUBRESOURCES);
 
 				const FD3D12RenderTargetView* RTV = nullptr;
 #if PLATFORM_REQUIRES_TYPELESS_RESOURCE_DISCARD_WORKAROUND
@@ -347,7 +338,7 @@ void FD3D12CommandContext::HandleResourceDiscardTransitions(
 			{
 				EnumerateSubresources(Resource, Info, Texture, [&](uint32 Subresource, const FD3D12RenderTargetView* RTV)
 				{
-					TransitionResource(Resource, InitialState, InitialState, Subresource);
+					TransitionResource(Resource, D3D12_RESOURCE_STATE_TBD, InitialState, Subresource);
 					ResourcesToDiscard.Emplace(Resource, Info.Flags, Subresource, Texture, RTV);
 				});
 			}
@@ -457,58 +448,7 @@ void FD3D12CommandContext::HandleTransientAliasing(const FD3D12TransitionData* T
 		if (Info.Action == FRHITransientAliasingInfo::EAction::Acquire)
 		{
 			TRACE_CPUPROFILER_EVENT_SCOPE(D3D12RHI::AcquireTransient);
-
-			if (GD3D12TransientAllocatorFullAliasingBarrier != 0)
-			{
-				AddAliasingBarrier(nullptr, Resource->GetResource());
-			}
-			else
-			{
-				for (const FRHITransientAliasingOverlap& Overlap : Info.Overlaps)
-				{
-					FD3D12Resource* ResourceBefore{};
-
-					switch (Overlap.Type)
-					{
-					case FRHITransientAliasingOverlap::EType::Texture:
-						{
-							const FD3D12Texture* Texture = RetrieveTexture(Overlap.Texture);
-							if (Texture)
-							{
-								ResourceBefore = Texture->GetResource();
-							}
-						}
-						break;
-					case FRHITransientAliasingOverlap::EType::Buffer:
-						{
-							const FD3D12Buffer* Buffer = RetrieveObject<FD3D12Buffer>(Overlap.Buffer);
-							if (Buffer)
-							{
-								ResourceBefore = Buffer->GetResource();
-							}
-						}
-						break;
-					}
-
-					// Resource may be null if this is a multi-GPU resource not present on the current GPU
-					check(ResourceBefore || GNumExplicitGPUsForRendering > 1);
-					if (ResourceBefore)
-					{
-						AddAliasingBarrier(ResourceBefore->GetResource(), Resource->GetResource());
-					}
-				}
-			}
-		}
-		else
-		{
-			TRACE_CPUPROFILER_EVENT_SCOPE(D3D12RHI::DiscardTransient);
-
-			// Restore the resource back to the initial state when done
-			D3D12_RESOURCE_STATES FinalState = GetInitialResourceState(Resource->GetDesc());
-			TransitionResource(Resource, D3D12_RESOURCE_STATE_TBD, FinalState, D3D12_RESOURCE_BARRIER_ALL_SUBRESOURCES);
-			
-			// Remove from caches
-			ClearShaderResources(BaseShaderResource, EShaderParameterTypeMask::SRVMask | EShaderParameterTypeMask::UAVMask);
+			AddAliasingBarrier(nullptr, Resource->GetResource());
 		}
 	}
 }
@@ -532,9 +472,11 @@ void FD3D12CommandContext::HandleResourceTransitions(const FD3D12TransitionData*
 		// and the transition SRV->UAV needs a UAV barrier then to work correctly otherwise there is no synchronization at all
 		bUAVBarrier |= bUAVAccessAfter;
 
+		const bool bHandleDiscardTransition = (Info.AccessAfter == ERHIAccess::Discard && QueueType == ED3D12QueueType::Direct);
+
 		// Process transitions which are forced during begin because those contain transition from Graphics to Compute and should
 		// help remove forced patch up command lists for async compute to run on the graphics queue
-		if (Info.Resource && ProcessTransitionDuringBegin(TransitionData))
+		if (Info.Resource && (ProcessTransitionDuringBegin(TransitionData) || bHandleDiscardTransition))
 		{
 			ProcessResource(*this, Info, [&](const FRHITransitionInfo& Info, FD3D12Resource* Resource, FD3D12Texture* UnusedTexture = nullptr)
 			{
@@ -543,12 +485,10 @@ void FD3D12CommandContext::HandleResourceTransitions(const FD3D12TransitionData*
 					return;
 				}
 
-				const bool bIsAsyncCompute = EnumHasAnyFlags(TransitionData->DstPipelines, ERHIPipeline::AsyncCompute);
-
 				// Use D3D12_RESOURCE_STATE_TBD as before state for now and don't use provided before state because there is not validation nor handling if the provided state
 				// doesn't match up with the tracked state - this needs to be improved and checked
 				D3D12_RESOURCE_STATES StateBefore = D3D12_RESOURCE_STATE_TBD;
-				D3D12_RESOURCE_STATES StateAfter = Info.AccessAfter == ERHIAccess::Discard ? GetInitialResourceState(Resource->GetDesc()) : GetD3D12ResourceState(Info.AccessAfter, bIsAsyncCompute);
+				D3D12_RESOURCE_STATES StateAfter = Info.AccessAfter == ERHIAccess::Discard ? GetInitialResourceState(Resource->GetDesc(), QueueType) : GetD3D12ResourceState(Info.AccessAfter, QueueType);
 
 				// enqueue the correct transitions
 				if (Info.IsWholeResource() || Resource->GetSubresourceCount() == 1)
diff --git a/Engine/Source/Runtime/D3D12RHI/Private/D3D12RHIPrivate.h b/Engine/Source/Runtime/D3D12RHI/Private/D3D12RHIPrivate.h
--- a/Engine/Source/Runtime/D3D12RHI/Private/D3D12RHIPrivate.h
+++ b/Engine/Source/Runtime/D3D12RHI/Private/D3D12RHIPrivate.h
@@ -216,7 +216,7 @@ private:
 	FProcessResult ProcessSubmissionQueue();
 	FProcessResult ProcessInterruptQueue();
 
-	static FD3D12CommandList* GenerateBarrierCommandListAndUpdateState(FD3D12CommandList* SourceCommandList);
+	static void GenerateBarrierCommandListAndUpdateState(FD3D12CommandList* SourceCommandList, TArray<FD3D12CommandList*, TInlineAllocator<2>>& OutBarrierCommandLists);
 
 	FCriticalSection ObjectsToDeleteCS;
 	TArray<FD3D12DeferredDeleteObject> ObjectsToDelete;
diff --git a/Engine/Source/Runtime/D3D12RHI/Private/D3D12Resources.h b/Engine/Source/Runtime/D3D12RHI/Private/D3D12Resources.h
--- a/Engine/Source/Runtime/D3D12RHI/Private/D3D12Resources.h
+++ b/Engine/Source/Runtime/D3D12RHI/Private/D3D12Resources.h
@@ -411,8 +411,7 @@ public:
 			// state then when a transition is required (will transition via scoped push/pop to requested state)
 			if (!bSRVOnly && InResourceState != ERHIAccess::Unknown && InResourceState != ERHIAccess::Discard)
 			{
-				bool bAsyncCompute = false;
-				return GetD3D12ResourceState(InResourceState, bAsyncCompute);
+				return GetD3D12ResourceState(InResourceState, ED3D12QueueType::Direct);
 			}
 			else
 			{
diff --git a/Engine/Source/Runtime/D3D12RHI/Private/D3D12Submission.cpp b/Engine/Source/Runtime/D3D12RHI/Private/D3D12Submission.cpp
--- a/Engine/Source/Runtime/D3D12RHI/Private/D3D12Submission.cpp
+++ b/Engine/Source/Runtime/D3D12RHI/Private/D3D12Submission.cpp
@@ -411,7 +411,10 @@ FD3D12DynamicRHI::FProcessResult FD3D12DynamicRHI::ProcessSubmissionQueue()
 					for (int32 Index = 0; Index < Payload->CommandListsToExecute.Num(); Index++)
 					{
 						FD3D12CommandList* CurrentCommandList = Payload->CommandListsToExecute[Index];
-						if (FD3D12CommandList* BarrierCommandList = GenerateBarrierCommandListAndUpdateState(CurrentCommandList))
+						TArray<FD3D12CommandList*, TInlineAllocator<2>> BarrierCommandLists;
+						GenerateBarrierCommandListAndUpdateState(CurrentCommandList, BarrierCommandLists);
+
+						for (FD3D12CommandList* BarrierCommandList : BarrierCommandLists)
 						{
 							FD3D12Queue& BarrierQueue = BarrierCommandList->Device->GetQueue(BarrierCommandList->QueueType);
 
@@ -575,11 +578,12 @@ FD3D12DynamicRHI::FProcessResult FD3D12DynamicRHI::ProcessSubmissionQueue()
 	return Result;
 }
 
-FD3D12CommandList* FD3D12DynamicRHI::GenerateBarrierCommandListAndUpdateState(FD3D12CommandList* SourceCommandList)
+void FD3D12DynamicRHI::GenerateBarrierCommandListAndUpdateState(FD3D12CommandList* SourceCommandList, TArray<FD3D12CommandList*, TInlineAllocator<2>>& OutBarrierCommandLists)
 {
 	TRACE_CPUPROFILER_EVENT_SCOPE(GenerateBarrierCommandListAndUpdateState);
 
 	FD3D12ResourceBarrierBatcher Batcher;
+	FD3D12ResourceBarrierBatcher BatcherForGraphicsToAsync;
 
 #if ENABLE_RESIDENCY_MANAGEMENT
 	TArray<FD3D12ResidencyHandle*> ResidencyHandles;
@@ -600,9 +604,11 @@ FD3D12CommandList* FD3D12DynamicRHI::GenerateBarrierCommandListAndUpdateState(FD
 		// We shouldn't have any TBD / CORRUPT states here
 		check(Before != D3D12_RESOURCE_STATE_TBD && Before != D3D12_RESOURCE_STATE_CORRUPT);
 		check(After  != D3D12_RESOURCE_STATE_TBD && After  != D3D12_RESOURCE_STATE_CORRUPT);
-		
+
 		if (Before != After)
 		{
+			FD3D12ResourceBarrierBatcher* CurrentBatcher = &Batcher;
+
 			if (SourceCommandList->QueueType != ED3D12QueueType::Direct)
 			{
 				check(!IsDirectQueueExclusiveD3D12State(After));
@@ -614,7 +620,7 @@ FD3D12CommandList* FD3D12DynamicRHI::GenerateBarrierCommandListAndUpdateState(FD
 					CResourceState& ResourceState_OnCommandList = SourceCommandList->GetResourceState_OnCommandList(PRB.Resource);
 					if (ResourceState_OnCommandList.HasInternalTransition() || !EnumHasAllFlags(Before, After))
 					{
-						bHasGraphicStates = true;
+						CurrentBatcher = &BatcherForGraphicsToAsync;
 					}
 					else
 					{
@@ -630,7 +636,7 @@ FD3D12CommandList* FD3D12DynamicRHI::GenerateBarrierCommandListAndUpdateState(FD
 
 			if (PRB.Resource->IsBackBuffer() && EnumHasAnyFlags(After, BackBufferBarrierWriteTransitionTargets))
 			{
-				Batcher.AddTransition(PRB.Resource, Before, After, PRB.SubResource);
+				CurrentBatcher->AddTransition(PRB.Resource, Before, After, PRB.SubResource);
 			}
 			// Special case for UAV access resources transitioning from UAV (then they need to transition from the cache hidden state instead)
 			else if (PRB.Resource->GetUAVAccessResource() && EnumHasAnyFlags(Before | After, D3D12_RESOURCE_STATE_UNORDERED_ACCESS))
@@ -639,18 +645,18 @@ FD3D12CommandList* FD3D12DynamicRHI::GenerateBarrierCommandListAndUpdateState(FD
 				check(!EnumHasAnyFlags(After, D3D12_RESOURCE_STATE_UNORDERED_ACCESS));
 					
 				// Add the aliasing barrier
-				Batcher.AddAliasingBarrier(PRB.Resource->GetUAVAccessResource(), PRB.Resource->GetResource());
+				CurrentBatcher->AddAliasingBarrier(PRB.Resource->GetUAVAccessResource(), PRB.Resource->GetResource());
 
 				D3D12_RESOURCE_STATES UAVState = ResourceState.GetUAVHiddenResourceState();
 				check(UAVState != D3D12_RESOURCE_STATE_TBD && !EnumHasAnyFlags(UAVState, D3D12_RESOURCE_STATE_UNORDERED_ACCESS));
 				if (UAVState != After)
 				{
-					Batcher.AddTransition(PRB.Resource, UAVState, After, PRB.SubResource);
+					CurrentBatcher->AddTransition(PRB.Resource, UAVState, After, PRB.SubResource);
 				}
 			}
 			else
 			{
-				Batcher.AddTransition(PRB.Resource, Before, After, PRB.SubResource);
+				CurrentBatcher->AddTransition(PRB.Resource, Before, After, PRB.SubResource);
 			}
 		}
 
@@ -694,33 +700,40 @@ FD3D12CommandList* FD3D12DynamicRHI::GenerateBarrierCommandListAndUpdateState(FD
 		}
 	}
 
-	if (Batcher.Num() == 0)
-		return nullptr;
-
-	// This command list requires a separate barrier command list to fix up tracked resource states.
-	TRACE_CPUPROFILER_EVENT_SCOPE(GetResourceBarrierCommandList);
+	const auto CreateBarrierCommandList = [&] (FD3D12Queue& Queue, FD3D12ResourceBarrierBatcher& Batcher)
+	{
+		// This command list requires a separate barrier command list to fix up tracked resource states.
+		TRACE_CPUPROFILER_EVENT_SCOPE(GetResourceBarrierCommandList);
 
-	FD3D12Queue& Queue = SourceCommandList->Device->GetQueue(
-		bHasGraphicStates
-			? ED3D12QueueType::Direct
-			: SourceCommandList->QueueType
-	);
+		// Get an allocator if we don't have one
+		if (!Queue.BarrierAllocator)
+		{
+			Queue.BarrierAllocator = Queue.Device->ObtainCommandAllocator(Queue.QueueType);
+		}
 
-	// Get an allocator if we don't have one
-	if (!Queue.BarrierAllocator)
-		Queue.BarrierAllocator = Queue.Device->ObtainCommandAllocator(Queue.QueueType);
+		// Get a new command list
+		FD3D12CommandList* BarrierCommandList = Queue.Device->ObtainCommandList(Queue.BarrierAllocator, &Queue.BarrierTimestamps, nullptr);
 
-	// Get a new command list
-	FD3D12CommandList* BarrierCommandList = Queue.Device->ObtainCommandList(Queue.BarrierAllocator, &Queue.BarrierTimestamps, nullptr);
+	#if ENABLE_RESIDENCY_MANAGEMENT
+		BarrierCommandList->UpdateResidency(ResidencyHandles);
+	#endif
 
-#if ENABLE_RESIDENCY_MANAGEMENT
-	BarrierCommandList->UpdateResidency(ResidencyHandles);
-#endif
-
-	Batcher.FlushIntoCommandList(*BarrierCommandList, Queue.BarrierTimestamps);
-	BarrierCommandList->Close();
+		Batcher.FlushIntoCommandList(*BarrierCommandList, Queue.BarrierTimestamps);
+		BarrierCommandList->Close();
+		return BarrierCommandList;
+	};
 
-	return BarrierCommandList;
+	if (Batcher.Num() > 0)
+	{
+		FD3D12Queue& Queue = SourceCommandList->Device->GetQueue(SourceCommandList->QueueType);
+		OutBarrierCommandLists.Emplace(CreateBarrierCommandList(Queue, Batcher));
+	}
+	
+	if (BatcherForGraphicsToAsync.Num() > 0)
+	{
+		FD3D12Queue& Queue = SourceCommandList->Device->GetQueue(ED3D12QueueType::Direct);
+		OutBarrierCommandLists.Emplace(CreateBarrierCommandList(Queue, BatcherForGraphicsToAsync));
+	}
 }
 
 uint64 FD3D12Queue::FinalizePayload(bool bRequiresSignal)
diff --git a/Engine/Source/Runtime/D3D12RHI/Private/D3D12TransientResourceAllocator.cpp b/Engine/Source/Runtime/D3D12RHI/Private/D3D12TransientResourceAllocator.cpp
--- a/Engine/Source/Runtime/D3D12RHI/Private/D3D12TransientResourceAllocator.cpp
+++ b/Engine/Source/Runtime/D3D12RHI/Private/D3D12TransientResourceAllocator.cpp
@@ -6,15 +6,15 @@
 #include "ProfilingDebugging/MemoryTrace.h"
 #include "ProfilingDebugging/AssetMetadataTrace.h"
 
-D3D12_RESOURCE_STATES GetInitialResourceState(const D3D12_RESOURCE_DESC& InDesc)
+D3D12_RESOURCE_STATES GetInitialResourceState(const D3D12_RESOURCE_DESC& InDesc, ED3D12QueueType QueueType)
 {
 	// Validate the creation state
 	D3D12_RESOURCE_STATES State = D3D12_RESOURCE_STATE_COMMON;
-	if (EnumHasAnyFlags(InDesc.Flags, D3D12_RESOURCE_FLAG_ALLOW_RENDER_TARGET))
+	if (EnumHasAnyFlags(InDesc.Flags, D3D12_RESOURCE_FLAG_ALLOW_RENDER_TARGET) && QueueType == ED3D12QueueType::Direct)
 	{
 		State = D3D12_RESOURCE_STATE_RENDER_TARGET;
 	}
-	else if (EnumHasAnyFlags(InDesc.Flags, D3D12_RESOURCE_FLAG_ALLOW_DEPTH_STENCIL))
+	else if (EnumHasAnyFlags(InDesc.Flags, D3D12_RESOURCE_FLAG_ALLOW_DEPTH_STENCIL) && QueueType == ED3D12QueueType::Direct)
 	{
 		State = D3D12_RESOURCE_STATE_DEPTH_WRITE;
 	}
@@ -157,7 +157,7 @@ FD3D12TransientResourceHeapAllocator::FD3D12TransientResourceHeapAllocator(FD3D1
 	, AllocationInfoQueryDevice(GetParentAdapter()->GetDevice(0))
 {}
 
-FRHITransientTexture* FD3D12TransientResourceHeapAllocator::CreateTexture(const FRHITextureCreateInfo& InCreateInfo, const TCHAR* InDebugName, uint32 InPassIndex)
+FRHITransientTexture* FD3D12TransientResourceHeapAllocator::CreateTexture(const FRHITextureCreateInfo& InCreateInfo, const TCHAR* InDebugName, const FRHITransientAllocationFences& Fences)
 {
 	FD3D12DynamicRHI* DynamicRHI = FD3D12DynamicRHI::GetD3DRHI();
 
@@ -166,14 +166,20 @@ FRHITransientTexture* FD3D12TransientResourceHeapAllocator::CreateTexture(const
 
 	Info.Alignment = FMath::Max<uint32>(D3D12_DEFAULT_RESOURCE_PLACEMENT_ALIGNMENT, Info.Alignment);
 
-	return CreateTextureInternal(InCreateInfo, InDebugName, InPassIndex, Info.SizeInBytes, Info.Alignment,
+	return CreateTextureInternal(InCreateInfo, InDebugName, Fences, Info.SizeInBytes, Info.Alignment,
 		[&](const FRHITransientHeap::FResourceInitializer& Initializer)
 	{
-		FResourceAllocatorAdapter ResourceAllocatorAdapter(GetParentAdapter(), static_cast<FD3D12TransientHeap&>(Initializer.Heap), Initializer.Allocation, Desc);
+		FD3D12TransientHeap& Heap = static_cast<FD3D12TransientHeap&>(Initializer.Heap);
 
-		FRHITextureCreateDesc CreateDesc(InCreateInfo, ERHIAccess::Discard, InDebugName);
-		FRHITexture* Texture = DynamicRHI->CreateD3D12Texture(CreateDesc, nullptr, &ResourceAllocatorAdapter);
-		return new FRHITransientTexture(Texture, ResourceAllocatorAdapter.GpuVirtualAddress, Initializer.Hash, Info.SizeInBytes, ERHITransientAllocationType::Heap, InCreateInfo);
+		return CreateTransientResource<FRHITransientTexture>([DynamicRHI, Adapter = GetParentAdapter(), &Heap, Allocation = Initializer.Allocation, Desc, InCreateInfo, InDebugName]
+		{
+			TRACE_CPUPROFILER_EVENT_SCOPE(AllocatePlacedTexture);
+			FResourceAllocatorAdapter ResourceAllocatorAdapter(Adapter, Heap, Allocation, Desc);
+			FRHITextureCreateDesc CreateDesc(InCreateInfo, ERHIAccess::Discard, InDebugName);
+			FRHITexture* Texture = DynamicRHI->CreateD3D12Texture(CreateDesc, nullptr, &ResourceAllocatorAdapter);
+			return FRHITransientResource::FResourceTaskResult{ Texture, ResourceAllocatorAdapter.GpuVirtualAddress };
+
+		}, Initializer.Hash, Info.SizeInBytes, InCreateInfo);
 	});
 }
 
@@ -210,7 +216,7 @@ void FD3D12TransientResourceHeapAllocator::FResourceAllocatorAdapter::AllocateRe
 #endif
 }
 
-FRHITransientBuffer* FD3D12TransientResourceHeapAllocator::CreateBuffer(const FRHIBufferCreateInfo& InCreateInfo, const TCHAR* InDebugName, uint32 InPassIndex)
+FRHITransientBuffer* FD3D12TransientResourceHeapAllocator::CreateBuffer(const FRHIBufferCreateInfo& InCreateInfo, const TCHAR* InDebugName, const FRHITransientAllocationFences& Fences)
 {
 	D3D12_RESOURCE_DESC Desc;
 	uint32 Alignment;
@@ -220,12 +226,19 @@ FRHITransientBuffer* FD3D12TransientResourceHeapAllocator::CreateBuffer(const FR
 	Alignment = D3D12_DEFAULT_RESOURCE_PLACEMENT_ALIGNMENT;
 	uint64 Size = Align(Desc.Width, Alignment);
 
-	return CreateBufferInternal(InCreateInfo, InDebugName, InPassIndex, Size, Alignment,
+	return CreateBufferInternal(InCreateInfo, InDebugName, Fences, Size, Alignment,
 		[&](const FRHITransientHeap::FResourceInitializer& Initializer)
 	{
-		FResourceAllocatorAdapter ResourceAllocatorAdapter(GetParentAdapter(), static_cast<FD3D12TransientHeap&>(Initializer.Heap), Initializer.Allocation, Desc);
-		FRHIBuffer* Buffer = FD3D12DynamicRHI::GetD3DRHI()->CreateBuffer(InCreateInfo, InDebugName, ERHIAccess::Discard, &ResourceAllocatorAdapter);
-		Buffer->SetTrackedAccess_Unsafe(ERHIAccess::Discard);
-		return new FRHITransientBuffer(Buffer, ResourceAllocatorAdapter.GpuVirtualAddress, Initializer.Hash, Size, ERHITransientAllocationType::Heap, InCreateInfo);
+		FD3D12TransientHeap& Heap = static_cast<FD3D12TransientHeap&>(Initializer.Heap);
+
+		return CreateTransientResource<FRHITransientBuffer>([Adapter = GetParentAdapter(), &Heap, Allocation = Initializer.Allocation, Desc, InCreateInfo, InDebugName]
+		{
+			TRACE_CPUPROFILER_EVENT_SCOPE(AllocatePlacedBuffer);
+			FResourceAllocatorAdapter ResourceAllocatorAdapter(Adapter, Heap, Allocation, Desc);
+			FRHIBuffer* Buffer = FD3D12DynamicRHI::GetD3DRHI()->CreateBuffer(InCreateInfo, InDebugName, ERHIAccess::Discard, &ResourceAllocatorAdapter);
+			Buffer->SetTrackedAccess_Unsafe(ERHIAccess::Discard);
+			return FRHITransientResource::FResourceTaskResult{ Buffer, ResourceAllocatorAdapter.GpuVirtualAddress };
+
+		}, Initializer.Hash, Size, InCreateInfo);
 	});
 }
\ No newline at end of file
diff --git a/Engine/Source/Runtime/D3D12RHI/Private/D3D12TransientResourceAllocator.h b/Engine/Source/Runtime/D3D12RHI/Private/D3D12TransientResourceAllocator.h
--- a/Engine/Source/Runtime/D3D12RHI/Private/D3D12TransientResourceAllocator.h
+++ b/Engine/Source/Runtime/D3D12RHI/Private/D3D12TransientResourceAllocator.h
@@ -5,7 +5,7 @@
 #include "D3D12Resources.h"
 #include "RHICoreTransientResourceAllocator.h"
 
-extern D3D12_RESOURCE_STATES GetInitialResourceState(const D3D12_RESOURCE_DESC& InDesc);
+extern D3D12_RESOURCE_STATES GetInitialResourceState(const D3D12_RESOURCE_DESC& InDesc, ED3D12QueueType QueueType);
 
 class FD3D12TransientHeap final
 	: public FRHITransientHeap
@@ -55,8 +55,9 @@ public:
 		default: checkNoEntry(); return false;
 		}
 	}
-	FRHITransientTexture* CreateTexture(const FRHITextureCreateInfo& InCreateInfo, const TCHAR* InDebugName, uint32 InPassIndex) override;
-	FRHITransientBuffer* CreateBuffer(const FRHIBufferCreateInfo& InCreateInfo, const TCHAR* InDebugName, uint32 InPassIndex) override;
+
+	FRHITransientTexture* CreateTexture(const FRHITextureCreateInfo& CreateInfo, const TCHAR* Name, const FRHITransientAllocationFences& Fences) override;
+	FRHITransientBuffer* CreateBuffer(const FRHIBufferCreateInfo& CreateInfo, const TCHAR* Name, const FRHITransientAllocationFences& Fences) override;
 
 private:
 
diff --git a/Engine/Source/Runtime/D3D12RHI/Private/D3D12Util.cpp b/Engine/Source/Runtime/D3D12RHI/Private/D3D12Util.cpp
--- a/Engine/Source/Runtime/D3D12RHI/Private/D3D12Util.cpp
+++ b/Engine/Source/Runtime/D3D12RHI/Private/D3D12Util.cpp
@@ -1624,7 +1624,7 @@ void LogResourceBarriers(TConstArrayView<D3D12_RESOURCE_BARRIER> Barriers, ID3D1
 }
 
 
-D3D12_RESOURCE_STATES GetD3D12ResourceState(ERHIAccess InRHIAccess, bool InIsAsyncCompute)
+D3D12_RESOURCE_STATES GetD3D12ResourceState(ERHIAccess InRHIAccess, ED3D12QueueType QueueType)
 {
 	// Add switch for common states (should cover all writeable states)
 	switch (InRHIAccess)
@@ -1658,7 +1658,7 @@ D3D12_RESOURCE_STATES GetD3D12ResourceState(ERHIAccess InRHIAccess, bool InIsAsy
 			D3D12_RESOURCE_STATES State = D3D12_RESOURCE_STATE_COMMON;
 
 			// Translate the requested after state to a D3D state
-			if (EnumHasAnyFlags(InRHIAccess, ERHIAccess::SRVGraphics) && !InIsAsyncCompute)
+			if (EnumHasAnyFlags(InRHIAccess, ERHIAccess::SRVGraphics) && QueueType == ED3D12QueueType::Direct)
 			{
 				State |= D3D12_RESOURCE_STATE_PIXEL_SHADER_RESOURCE | D3D12_RESOURCE_STATE_NON_PIXEL_SHADER_RESOURCE;
 			}
diff --git a/Engine/Source/Runtime/D3D12RHI/Private/D3D12Util.h b/Engine/Source/Runtime/D3D12RHI/Private/D3D12Util.h
--- a/Engine/Source/Runtime/D3D12RHI/Private/D3D12Util.h
+++ b/Engine/Source/Runtime/D3D12RHI/Private/D3D12Util.h
@@ -9,6 +9,7 @@
 #include "CoreMinimal.h"
 #include "Containers/Queue.h"
 #include "D3D12RHICommon.h"
+#include "D3D12Queue.h"
 #include "DXGIUtilities.h"
 #include "RenderUtils.h"
 #include "ShaderCore.h"
@@ -236,7 +237,7 @@ static bool IsDirectQueueExclusiveD3D12State(D3D12_RESOURCE_STATES InState)
 	return EnumHasAnyFlags(InState, D3D12_RESOURCE_STATE_RENDER_TARGET | D3D12_RESOURCE_STATE_DEPTH_WRITE | D3D12_RESOURCE_STATE_DEPTH_READ | D3D12_RESOURCE_STATE_PIXEL_SHADER_RESOURCE);
 }
 
-D3D12_RESOURCE_STATES GetD3D12ResourceState(ERHIAccess InRHIAccess, bool InIsAsyncCompute);
+D3D12_RESOURCE_STATES GetD3D12ResourceState(ERHIAccess InRHIAccess, ED3D12QueueType QueueType);
 
 //==================================================================================================================================
 // CResourceState
diff --git a/Engine/Source/Runtime/RHI/Private/RHI.cpp b/Engine/Source/Runtime/RHI/Private/RHI.cpp
--- a/Engine/Source/Runtime/RHI/Private/RHI.cpp
+++ b/Engine/Source/Runtime/RHI/Private/RHI.cpp
@@ -2106,20 +2106,18 @@ void FRHIBufferViewCache::SetDebugName(FRHICommandListBase& RHICmdList, const TC
 
 #endif
 
-void FRHITransientTexture::Acquire(FRHICommandListBase& RHICmdList, const TCHAR* InName, uint32 InPassIndex, uint64 InAcquireCycle)
+void FRHITransientTexture::BindDebugLabelName(FRHICommandListBase& RHICmdList)
 {
-	FRHITransientResource::Acquire(RHICmdList, InName, InPassIndex, InAcquireCycle);
-	ViewCache.SetDebugName(RHICmdList, InName);
+	ViewCache.SetDebugName(RHICmdList, GetName());
 
 #if !(UE_BUILD_SHIPPING || UE_BUILD_TEST)
-	RHICmdList.BindDebugLabelName(GetRHI(), InName);
+	RHICmdList.BindDebugLabelName(GetRHI(), GetName());
 #endif
 }
 
-void FRHITransientBuffer::Acquire(FRHICommandListBase& RHICmdList, const TCHAR* InName, uint32 InPassIndex, uint64 InAcquireCycle)
+void FRHITransientBuffer::BindDebugLabelName(FRHICommandListBase& RHICmdList)
 {
-	FRHITransientResource::Acquire(RHICmdList, InName, InPassIndex, InAcquireCycle);
-	ViewCache.SetDebugName(RHICmdList, InName);
+	ViewCache.SetDebugName(RHICmdList, GetName());
 
 	// TODO: Add method to rename a buffer.
 }
diff --git a/Engine/Source/Runtime/RHI/Private/RHIValidation.cpp b/Engine/Source/Runtime/RHI/Private/RHIValidation.cpp
--- a/Engine/Source/Runtime/RHI/Private/RHIValidation.cpp
+++ b/Engine/Source/Runtime/RHI/Private/RHIValidation.cpp
@@ -586,7 +586,6 @@ void FValidationRHI::RHICreateTransition(FRHITransition* Transition, const FRHIT
 		if (Info.IsAcquire())
 		{
 			checkf(Resource->TransientState.bTransient, TEXT("Acquiring resource %s which is not transient. Only transient resources can be acquired."), Resource->GetDebugName());
-			checkf(SrcPipelines == ERHIPipeline::Graphics, TEXT("Acquiring a transient resource (%s) must begin on the graphics pipe."), Resource->GetDebugName());
 
 			AliasingOps.Emplace(FOperation::AcquireTransientResource(Resource, nullptr));
 
@@ -608,14 +607,6 @@ void FValidationRHI::RHICreateTransition(FRHITransition* Transition, const FRHIT
 				AliasingOverlapOps.Emplace(FOperation::AliasingOverlap(ResourceBefore, Resource, nullptr));
 			}
 		}
-		else
-		{
-			checkf(Info.Overlaps.IsEmpty(), TEXT("Aliasing overlaps provided on a Discard for resource %s. Overlaps must be provided on an acquire."), Resource->GetDebugName());
-			checkf(Resource->TransientState.bTransient, TEXT("Discarding resource %s which is not transient. Only transient resources can be discarded."), Resource->GetDebugName());
-			checkf(DstPipelines == ERHIPipeline::Graphics, TEXT("Discarding a transient resource (%s) must end on the graphics pipe."), Resource->GetDebugName());
-
-			AliasingOps.Emplace(FOperation::DiscardTransientResource(Resource, nullptr));
-		}
 	}
 
 	for (int32 Index = 0; Index < CreateInfo.TransitionInfos.Num(); ++Index)
@@ -688,9 +679,6 @@ void FValidationRHI::RHICreateTransition(FRHITransition* Transition, const FRHIT
 			case EOpType::AcquireTransient:
 				Op.Data_AcquireTransient.CreateBacktrace = Backtrace;
 				break;
-			case EOpType::DiscardTransient:
-				Op.Data_DiscardTransient.CreateBacktrace = Backtrace;
-				break;
 			}
 		}
 
@@ -1289,30 +1277,30 @@ namespace RHIValidation
 			*GetReasonString_DuplicateBackTrace(PreviousAcquireTrace, CurrentAcquireTrace));
 	}
 
-	static inline FString GetReasonString_DuplicateDiscardTransient(FResource* Resource, void* PreviousDiscardTrace, void* CurrentDiscardTrace)
+	static inline FString GetReasonString_DiscardWithoutAcquireTransient(FResource* Resource, void* DiscardTrace)
 	{
 		FString DebugName = GetResourceDebugName(Resource, {});
 		return FString::Printf(
 			BARRIER_TRACKER_LOG_PREFIX_RESNAME
-			TEXT("Mismatched discard of transient resource %s. A transient resource may only be discarded once in its lifetime.\n")
+			TEXT("Attempted to discard transient resource %s, but it was never acquired.\n")
 			TEXT("%s")
 			BARRIER_TRACKER_LOG_SUFFIX,
 			*DebugName,
 			*DebugName,
-			*GetReasonString_DuplicateBackTrace(PreviousDiscardTrace, CurrentDiscardTrace));
+			*GetReasonString_Backtrace(TEXT("discard"), TEXT("RHICreateTransition"), DiscardTrace));
 	}
-
-	static inline FString GetReasonString_DiscardWithoutAcquireTransient(FResource* Resource, void* DiscardTrace)
+	
+	static inline FString GetReasonString_AlreadyDiscarded(FResource* Resource, void* DiscardTrace)
 	{
 		FString DebugName = GetResourceDebugName(Resource, {});
 		return FString::Printf(
 			BARRIER_TRACKER_LOG_PREFIX_RESNAME
-			TEXT("Attempted to discard transient resource %s, but it was never acquired.\n")
+			TEXT("Attempted to transition transient resource %s to ERHIAccess::Discard, but it has already been discarded.\n")
 			TEXT("%s")
 			BARRIER_TRACKER_LOG_SUFFIX,
 			*DebugName,
 			*DebugName,
-			*GetReasonString_Backtrace(TEXT("acquire"), TEXT("RHICreateTransition"), DiscardTrace));
+			*GetReasonString_Backtrace(TEXT("discard"), TEXT("RHICreateTransition"), DiscardTrace));
 	}
 
 	static inline FString GetReasonString_DuplicateBeginTransition(
@@ -1530,28 +1518,33 @@ namespace RHIValidation
 			AcquireBacktrace = CreateTrace;
 		}
 
+		NumAcquiredSubresources = Resource->GetNumSubresources() * GetRHIPipelineCount();
+
 		if (Resource->LoggingMode != ELoggingMode::None)
 		{
 			Log(Resource, {}, CreateTrace, TEXT("Acquire"), TEXT("Acquire"), TEXT("Transient Acquire"));
 		}
 	}
 
-	void FTransientState::Discard(FResource* Resource, void* CreateTrace)
+	void FTransientState::Discard(FResource* Resource, void* CreateTrace, ERHIPipeline DiscardPipelines)
 	{
 		RHI_VALIDATION_CHECK(bTransient, *GetReasonString_DiscardNonTransient(Resource));
-
 		RHI_VALIDATION_CHECK(Status != EStatus::None, *GetReasonString_DiscardWithoutAcquireTransient(Resource, CreateTrace));
-		RHI_VALIDATION_CHECK(Status != EStatus::Discarded, *GetReasonString_DuplicateDiscardTransient(Resource, DiscardBacktrace, CreateTrace));
-		Status = EStatus::Discarded;
+		RHI_VALIDATION_CHECK(Status != EStatus::Discarded, *GetReasonString_AlreadyDiscarded(Resource, CreateTrace));
 
-		if (!DiscardBacktrace)
-		{
-			DiscardBacktrace = CreateTrace;
-		}
+		// When discarding from all pipes, each pipe will call Discard separately. Otherwise it's just one call.
+		const uint32 NumDerefs = DiscardPipelines == ERHIPipeline::All ? 1 : 2;
 
-		if (Resource->LoggingMode != ELoggingMode::None)
+		NumAcquiredSubresources -= NumDerefs;
+
+		if (NumAcquiredSubresources == 0)
 		{
-			Log(Resource, {}, CreateTrace, TEXT("Discard"), TEXT("Discard"), TEXT("Transient Discard"));
+			Status = EStatus::Discarded;
+
+			if (Resource->LoggingMode != ELoggingMode::None)
+			{
+				Log(Resource, {}, CreateTrace, TEXT("Discard"), TEXT("Discard"), TEXT("Transient Discard"));
+			}
 		}
 	}
 
@@ -1622,7 +1615,12 @@ namespace RHIValidation
 
 		if (Resource->TransientState.bTransient)
 		{
-			RHI_VALIDATION_CHECK(Resource->TransientState.IsAcquired() || (Resource->TransientState.IsDiscarded() && TargetState.Access == ERHIAccess::Discard), *GetReasonString_TransitionWithoutAcquire(Resource));
+			RHI_VALIDATION_CHECK(Resource->TransientState.IsAcquired(), *GetReasonString_TransitionWithoutAcquire(Resource));
+
+			if (TargetState.Access == ERHIAccess::Discard)
+			{
+				Resource->TransientState.Discard(Resource, CreateTrace, CurrentStateFromRHI.Pipelines);
+			}
 		}
 
 		// Check we're not already transitioning
@@ -1692,11 +1690,6 @@ namespace RHIValidation
 		State.bTransitioning = false;
 		State.BeginTransitionBacktrace = nullptr;
 
-		if (Resource->TransientState.bTransient)
-		{
-			RHI_VALIDATION_CHECK(Resource->TransientState.IsAcquired() || (Resource->TransientState.IsDiscarded() && TargetState.Access == ERHIAccess::Discard), *GetReasonString_TransitionWithoutAcquire(Resource));
-		}
-
 		// Check that the end matches the begin.
 		RHI_VALIDATION_CHECK(TargetState == State.Current, *GetReasonString_MismatchedEndTransition(Resource, SubresourceIndex, State.Current, TargetState));
 
@@ -2005,11 +1998,6 @@ namespace RHIValidation
 			Data_AcquireTransient.Resource->ReleaseOpRef();
 			break;
 
-		case EOpType::DiscardTransient:
-			Data_DiscardTransient.Resource->TransientState.Discard(Data_DiscardTransient.Resource, Data_DiscardTransient.CreateBacktrace);
-			Data_AcquireTransient.Resource->ReleaseOpRef();
-			break;
-
 		case EOpType::InitTransient:
 			Data_InitTransient.Resource->InitTransient(Data_InitTransient.DebugName);
 			delete[] Data_InitTransient.DebugName;
@@ -2587,11 +2575,16 @@ FValidationTransientResourceAllocator::~FValidationTransientResourceAllocator()
 	checkf(!RHIAllocator, TEXT("Release was not called on FRHITransientResourceAllocator."));
 }
 
-FRHITransientTexture* FValidationTransientResourceAllocator::CreateTexture(const FRHITextureCreateInfo& InCreateInfo, const TCHAR* InDebugName, uint32 InPassIndex)
+void FValidationTransientResourceAllocator::SetCreateMode(ERHITransientResourceCreateMode InCreateMode)
+{
+	// Validation intentionally doesn't pass through the create mode. It's always inline.
+}
+
+FRHITransientTexture* FValidationTransientResourceAllocator::CreateTexture(const FRHITextureCreateInfo& InCreateInfo, const TCHAR* InDebugName, const FRHITransientAllocationFences& Fences)
 {
 	check(FRHITextureCreateInfo::CheckValidity(InCreateInfo, InDebugName));
 
-	FRHITransientTexture* TransientTexture = RHIAllocator->CreateTexture(InCreateInfo, InDebugName, InPassIndex);
+	FRHITransientTexture* TransientTexture = RHIAllocator->CreateTexture(InCreateInfo, InDebugName, Fences);
 
 	if (!TransientTexture)
 	{
@@ -2623,9 +2616,9 @@ FRHITransientTexture* FValidationTransientResourceAllocator::CreateTexture(const
 	return TransientTexture;
 }
 
-FRHITransientBuffer* FValidationTransientResourceAllocator::CreateBuffer(const FRHIBufferCreateInfo& InCreateInfo, const TCHAR* InDebugName, uint32 InPassIndex)
+FRHITransientBuffer* FValidationTransientResourceAllocator::CreateBuffer(const FRHIBufferCreateInfo& InCreateInfo, const TCHAR* InDebugName, const FRHITransientAllocationFences& Fences)
 {
-	FRHITransientBuffer* TransientBuffer = RHIAllocator->CreateBuffer(InCreateInfo, InDebugName, InPassIndex);
+	FRHITransientBuffer* TransientBuffer = RHIAllocator->CreateBuffer(InCreateInfo, InDebugName, Fences);
 
 	if (!TransientBuffer)
 	{
@@ -2651,21 +2644,21 @@ FRHITransientBuffer* FValidationTransientResourceAllocator::CreateBuffer(const F
 	return TransientBuffer;
 }
 
-void FValidationTransientResourceAllocator::DeallocateMemory(FRHITransientTexture* InTransientTexture, uint32 InPassIndex)
+void FValidationTransientResourceAllocator::DeallocateMemory(FRHITransientTexture* InTransientTexture, const FRHITransientAllocationFences& Fences)
 {
 	check(InTransientTexture);
 
-	RHIAllocator->DeallocateMemory(InTransientTexture, InPassIndex);
+	RHIAllocator->DeallocateMemory(InTransientTexture, Fences);
 
 	checkf(AllocatedResourceMap.Contains(InTransientTexture->GetRHI()), TEXT("DeallocateMemory called on texture %s, but it is not marked as allocated."), InTransientTexture->GetName());
 	AllocatedResourceMap.Remove(InTransientTexture->GetRHI());
 }
 
-void FValidationTransientResourceAllocator::DeallocateMemory(FRHITransientBuffer* InTransientBuffer, uint32 InPassIndex)
+void FValidationTransientResourceAllocator::DeallocateMemory(FRHITransientBuffer* InTransientBuffer, const FRHITransientAllocationFences& Fences)
 {
 	check(InTransientBuffer);
 
-	RHIAllocator->DeallocateMemory(InTransientBuffer, InPassIndex);
+	RHIAllocator->DeallocateMemory(InTransientBuffer, Fences);
 
 	checkf(AllocatedResourceMap.Contains(InTransientBuffer->GetRHI()), TEXT("DeallocateMemory called on buffer %s, but it is not marked as allocated."), InTransientBuffer->GetName());
 	AllocatedResourceMap.Remove(InTransientBuffer->GetRHI());
@@ -2692,27 +2685,6 @@ void FValidationTransientResourceAllocator::Flush(FRHICommandListImmediate& RHIC
 
 void FValidationTransientResourceAllocator::Release(FRHICommandListImmediate& RHICmdList)
 {
-	// Check all allocated resource data and make sure all memory is freed again
-	{
-		if (AllocatedResourceMap.Num() > 0)
-		{
-			FString ErrorMessage = FString::Printf(
-				TRANSIENT_RESOURCE_LOG_PREFIX_REASON("Open transient allocations")
-				TEXT("%d Transient Resource allocations still have memory allocated. Call 'DeallocateMemory' on all transient allocated resources prior to releasing the allocator.\n\n")
-				TEXT("Resources with Allocated Memory:\n"),
-				AllocatedResourceMap.Num());
-
-			for (const auto& KeyValue : AllocatedResourceMap)
-			{
-				const FAllocatedResourceData& ResourceData = KeyValue.Value;
-
-				ErrorMessage += FString::Printf(TEXT("         %s (%s)\n"), *ResourceData.DebugName, ResourceData.ResourceType == FAllocatedResourceData::EType::Texture ? TEXT("Texture") : TEXT("Buffer"));
-			}
-			ErrorMessage += FString::Printf(TRANSIENT_RESOURCE_LOG_SUFFIX);
-			FValidationRHI::ReportValidationFailure(*ErrorMessage);
-		}
-	}
-
 	RHIAllocator->Release(RHICmdList);
 	RHIAllocator = nullptr;
 	delete this;
diff --git a/Engine/Source/Runtime/RHI/Public/RHIGlobals.h b/Engine/Source/Runtime/RHI/Public/RHIGlobals.h
--- a/Engine/Source/Runtime/RHI/Public/RHIGlobals.h
+++ b/Engine/Source/Runtime/RHI/Public/RHIGlobals.h
@@ -254,6 +254,9 @@ struct FRHIGlobals
 	/** True if the RHI and current hardware supports efficient AsyncCompute (by default we assume false and later we can enable this for more hardware) */
 	bool SupportsEfficientAsyncCompute = false;
 
+	/** True if the RHI supports aliasing transient resources on the async compute pipe. */
+	bool SupportsAsyncComputeTransientAliasing = false;
+
 	/** True if the RHI supports getting the result of occlusion queries when on a thread other than the render thread */
 	bool SupportsParallelOcclusionQueries = false;
 
diff --git a/Engine/Source/Runtime/RHI/Public/RHIPipeline.h b/Engine/Source/Runtime/RHI/Public/RHIPipeline.h
--- a/Engine/Source/Runtime/RHI/Public/RHIPipeline.h
+++ b/Engine/Source/Runtime/RHI/Public/RHIPipeline.h
@@ -61,6 +61,16 @@ class TRHIPipelineArray : public TStaticArray<ElementType, GetRHIPipelineCount()
 public:
 	using Base::Base;
 
+	FORCEINLINE ElementType& operator[](int32 Index)
+	{
+		return Base::operator[](Index);
+	}
+
+	FORCEINLINE const ElementType& operator[](int32 Index) const
+	{
+		return Base::operator[](Index);
+	}
+
 	FORCEINLINE ElementType& operator[](ERHIPipeline Pipeline)
 	{
 		return Base::operator[](GetRHIPipelineIndex(Pipeline));
@@ -70,4 +80,4 @@ public:
 	{
 		return Base::operator[](GetRHIPipelineIndex(Pipeline));
 	}
-};
+};
\ No newline at end of file
diff --git a/Engine/Source/Runtime/RHI/Public/RHITransientResourceAllocator.h b/Engine/Source/Runtime/RHI/Public/RHITransientResourceAllocator.h
--- a/Engine/Source/Runtime/RHI/Public/RHITransientResourceAllocator.h
+++ b/Engine/Source/Runtime/RHI/Public/RHITransientResourceAllocator.h
@@ -4,6 +4,7 @@
 
 #include "RHI.h"
 #include "RHIResources.h"
+#include "Tasks/Task.h"
 
 class FRHICommandListBase;
 class FRHICommandListImmediate;
@@ -83,6 +84,14 @@ class FRHITransientResource
 public:
 	static const uint32 kInvalidPassIndex = TNumericLimits<uint32>::Max();
 
+	struct FResourceTaskResult
+	{
+		TRefCountPtr<FRHIResource> Resource; 
+		uint64 GpuVirtualAddress;
+	};
+
+	using FResourceTask = UE::Tasks::TTask<FResourceTaskResult>;
+
 	FRHITransientResource(
 		FRHIResource* InResource,
 		uint64 InGpuVirtualAddress,
@@ -98,34 +107,66 @@ public:
 		, ResourceType(InResourceType)
 	{}
 
+	FRHITransientResource(
+		const FResourceTask& InResourceTask,
+		uint64 InHash,
+		uint64 InSize,
+		ERHITransientAllocationType InAllocationType,
+		ERHITransientResourceType InResourceType)
+		: ResourceTask(InResourceTask)
+		, Hash(InHash)
+		, Size(InSize)
+		, AllocationType(InAllocationType)
+		, ResourceType(InResourceType)
+	{}
+
 	virtual ~FRHITransientResource() = default;
 
 	//////////////////////////////////////////////////////////////////////////////////////////////////
 	//! Internal Allocator API
 
-	virtual void Acquire(FRHICommandListBase& RHICmdList, const TCHAR* InName, uint32 InAcquirePassIndex, uint64 InAllocatorCycle)
+	void Acquire(const TCHAR* InName, uint32 InAcquirePassIndex, uint64 InAllocatorCycle)
 	{
 		Name = InName;
 		AcquirePasses = TInterval<uint32>(0, InAcquirePassIndex);
-		DiscardPasses = TInterval<uint32>(kInvalidPassIndex, kInvalidPassIndex);
+		DiscardPass = kInvalidPassIndex;
+		bAcquired = true;
 		AcquireCycle = InAllocatorCycle;
 		AcquireCount++;
 		AliasingOverlaps.Reset();
 	}
 
-	void Discard(uint32 InDiscardPassIndex)
+	void Discard()
 	{
-		DiscardPasses.Min = InDiscardPassIndex;
+		bAcquired = false;
 	}
 
-	void AddAliasingOverlap(FRHITransientResource* InResource)
+	void AddAliasingOverlap(FRHITransientResource* InBeforeResource, uint32 InAcquirePassIndex)
 	{
-		AliasingOverlaps.Emplace(InResource->GetRHI(), InResource->IsTexture() ? FRHITransientAliasingOverlap::EType::Texture : FRHITransientAliasingOverlap::EType::Buffer);
+		check(!InBeforeResource->IsAcquired());
+
+		// Aliasing overlaps are currently only tracked with RHI validation, as no RHI is actually using them.
+		if (GRHIValidationEnabled)
+		{
+			AliasingOverlaps.Emplace(InBeforeResource->GetRHI(), InBeforeResource->IsTexture() ? FRHITransientAliasingOverlap::EType::Texture : FRHITransientAliasingOverlap::EType::Buffer);
+		}
 
-		check(InResource->DiscardPasses.Min != kInvalidPassIndex);
+		InBeforeResource->DiscardPass = FMath::Min(InBeforeResource->DiscardPass, AcquirePasses.Max);
+		AcquirePasses.Min = FMath::Max(AcquirePasses.Min, InAcquirePassIndex);
 
-		InResource->DiscardPasses.Max = FMath::Min(InResource->DiscardPasses.Max,             AcquirePasses.Max);
-		            AcquirePasses.Min = FMath::Max(            AcquirePasses.Min, InResource->DiscardPasses.Min);
+		check(AcquirePasses.Min <= AcquirePasses.Max);
+	}
+
+	void Finish(FRHICommandListBase& RHICmdList)
+	{
+		if (ResourceTask.IsValid())
+		{
+			FResourceTaskResult Result = MoveTemp(ResourceTask.GetResult());
+			Resource = MoveTemp(Result.Resource);
+			GpuVirtualAddress = Result.GpuVirtualAddress;
+			ResourceTask = {};
+		}
+		BindDebugLabelName(RHICmdList);
 	}
 
 	FRHITransientHeapAllocation& GetHeapAllocation()
@@ -155,11 +196,14 @@ public:
 	//////////////////////////////////////////////////////////////////////////////////////////////////
 
 	// Returns the underlying RHI resource.
-	FRHIResource* GetRHI() const { return Resource; }
+	FRHIResource* GetRHI() const { check(!ResourceTask.IsValid()); return Resource; }
 
 	// Returns the gpu virtual address of the transient resource.
 	uint64 GetGpuVirtualAddress() const { return GpuVirtualAddress; }
 
+	// Returns whether a resource has a pending task.
+	bool HasResourceTask() const { return ResourceTask.IsValid(); }
+
 	// Returns the name assigned to the transient resource at allocation time.
 	const TCHAR* GetName() const { return Name; }
 
@@ -179,13 +223,14 @@ public:
 	TConstArrayView<FRHITransientAliasingOverlap> GetAliasingOverlaps() const { return AliasingOverlaps; }
 
 	// Returns the pass index which may end acquiring this resource.
-	TInterval<uint32> GetAcquirePasses() const { return AcquirePasses; }
+	uint32 GetAcquirePass() const { return AcquirePasses.Min; }
 
 	// Returns the pass index which discarded this resource.
-	TInterval<uint32> GetDiscardPasses() const { return DiscardPasses; }
+	uint32 GetDiscardPass() const { return DiscardPass; }
 
 	// Returns whether this resource is still in an acquired state.
-	bool IsAcquired() const { return DiscardPasses.Min == kInvalidPassIndex; }
+	bool IsAcquired() const { return bAcquired; }
+	bool IsDiscarded() const { return !bAcquired; }
 
 	ERHITransientResourceType GetResourceType() const { return ResourceType; }
 
@@ -198,8 +243,11 @@ public:
 	bool IsPageAllocated() const { return AllocationType == ERHITransientAllocationType::Page; }
 
 private:
+	virtual void BindDebugLabelName(FRHICommandListBase& RHICmdList) = 0;
+
 	// Underlying RHI resource.
 	TRefCountPtr<FRHIResource> Resource;
+	FResourceTask ResourceTask;
 
 	// The Gpu virtual address of the RHI resource.
 	uint64 GpuVirtualAddress = 0;
@@ -230,7 +278,8 @@ private:
 
 	// Start -> End split pass index intervals for acquire / discard operations.
 	TInterval<uint32> AcquirePasses = TInterval<uint32>(0, 0);
-	TInterval<uint32> DiscardPasses = TInterval<uint32>(0, 0);
+	uint32 DiscardPass = 0;
+	bool bAcquired = false;
 
 	ERHITransientAllocationType AllocationType;
 	ERHITransientResourceType ResourceType;
@@ -240,7 +289,17 @@ class FRHITransientTexture final : public FRHITransientResource
 {
 public:
 	FRHITransientTexture(
-		FRHITexture* InTexture,
+		const FResourceTask& InResourceTask,
+		uint64 InHash,
+		uint64 InSize,
+		ERHITransientAllocationType InAllocationType,
+		const FRHITextureCreateInfo& InCreateInfo)
+		: FRHITransientResource(InResourceTask, InHash, InSize, InAllocationType, ERHITransientResourceType::Texture)
+		, CreateInfo(InCreateInfo)
+	{}
+
+	FRHITransientTexture(
+		FRHIResource* InTexture,
 		uint64 InGpuVirtualAddress,
 		uint64 InHash,
 		uint64 InSize,
@@ -250,11 +309,6 @@ public:
 		, CreateInfo(InCreateInfo)
 	{}
 
-	//////////////////////////////////////////////////////////////////////////////////////////////////
-	//! Internal Allocator API
-	RHI_API void Acquire(FRHICommandListBase& RHICmdList, const TCHAR* InName, uint32 InAcquirePassIndex, uint64 InInitCycle) override;
-	//////////////////////////////////////////////////////////////////////////////////////////////////
-
 	// Returns the underlying RHI texture.
 	FRHITexture* GetRHI() const { return static_cast<FRHITexture*>(FRHITransientResource::GetRHI()); }
 
@@ -272,13 +326,26 @@ public:
 
 	// The persistent view cache containing all views created for this texture.
 	FRHITextureViewCache ViewCache;
+
+private:
+	RHI_API void BindDebugLabelName(FRHICommandListBase& RHICmdList) override;
 };
 
 class FRHITransientBuffer final : public FRHITransientResource
 {
 public:
 	FRHITransientBuffer(
-		FRHIBuffer* InBuffer,
+		const FResourceTask& InResourceTask,
+		uint64 InHash,
+		uint64 InSize,
+		ERHITransientAllocationType InAllocationType,
+		const FRHIBufferCreateInfo& InCreateInfo)
+		: FRHITransientResource(InResourceTask, InHash, InSize, InAllocationType, ERHITransientResourceType::Buffer)
+		, CreateInfo(InCreateInfo)
+	{}
+
+	FRHITransientBuffer(
+		FRHIResource* InBuffer,
 		uint64 InGpuVirtualAddress,
 		uint64 InHash,
 		uint64 InSize,
@@ -288,11 +355,6 @@ public:
 		, CreateInfo(InCreateInfo)
 	{}
 
-	//////////////////////////////////////////////////////////////////////////////////////////////////
-	//! Internal Allocator API
-	RHI_API void Acquire(FRHICommandListBase& RHICmdList, const TCHAR* InName, uint32 InAcquirePassIndex, uint64 InInitCycle) override;
-	//////////////////////////////////////////////////////////////////////////////////////////////////
-
 	// Returns the underlying RHI buffer.
 	FRHIBuffer* GetRHI() const { return static_cast<FRHIBuffer*>(FRHITransientResource::GetRHI()); }
 
@@ -310,6 +372,9 @@ public:
 
 	// The persistent view cache containing all views created for this buffer.
 	FRHIBufferViewCache ViewCache;
+
+private:
+	RHI_API void BindDebugLabelName(FRHICommandListBase& RHICmdList) override;
 };
 
 class FRHITransientAllocationStats
@@ -350,21 +415,108 @@ public:
 
 ENUM_CLASS_FLAGS(FRHITransientAllocationStats::EMemoryRangeFlags);
 
+/** This data structure contains fence values used for allocating / deallocating transient memory regions for transient resources. A memory region can
+ *  be re-used if the deallocation fences from the discarding resource and the allocation fences for the acquiring resource are not executing simultaneously
+ *  on both the graphics | async compute pipe on the GPU timeline.
+ *
+ *  Allocation events are always on a single pipeline, while deallocation events can happen on multiple pipelines at the same time. Async compute is represented
+ *  using three fence values: one for the async compute pipe, and two for the fork / join points on the graphics pipe. If fences are active on both pipes at the
+ *  same time, the graphics fence must be contained within the async compute fork / join region.
+ */
+class FRHITransientAllocationFences
+{
+public:
+	// Returns the fence at which the Acquire operation can occur for the given pair of resources transitioning from Discard -> Acquire.
+	static uint32 GetAcquireFence(const FRHITransientAllocationFences& Discard, const FRHITransientAllocationFences& Acquire)
+	{
+		check(Acquire.IsSinglePipeline());
+
+		// Graphics -> Graphics | AsyncCompute
+		if (Discard.Graphics != Invalid && Discard.AsyncCompute == Invalid)
+		{
+			return Discard.Graphics;
+		}
+
+		// All | Async ->
+		return Acquire.AsyncCompute != Invalid
+			// Async Compute
+			? Discard.AsyncCompute
+			// Graphics
+			: Discard.GraphicsForkJoin.Max;
+	}
+
+	// Returns whether two regions described by the discard and acquire fences contain each other. If they do, that means the memory would being used by both pipes simultaneously and cannot be aliased.
+	static bool Contains(const FRHITransientAllocationFences& Discard, const FRHITransientAllocationFences& Acquire)
+	{
+		return Contains(Discard.GraphicsForkJoin, Acquire.Graphics) || Contains(Acquire.GraphicsForkJoin, Discard.Graphics);
+	}
+
+	FRHITransientAllocationFences() = default;
+
+	void SetGraphics(uint32 InGraphics)
+	{
+		check(!GraphicsForkJoin.IsValid() || Contains(GraphicsForkJoin, InGraphics));
+		Graphics = InGraphics;
+	}
+
+	void SetAsyncCompute(uint32 InAsyncCompute, TInterval<uint32> InGraphicsForkJoin)
+	{
+		check(InGraphicsForkJoin.IsValid() && Contains(InGraphicsForkJoin, InAsyncCompute));
+		check(Graphics == Invalid || Contains(InGraphicsForkJoin, Graphics));
+		AsyncCompute = InAsyncCompute;
+		GraphicsForkJoin = InGraphicsForkJoin;
+	}
+
+	uint32 GetSinglePipeline() const
+	{
+		check(IsSinglePipeline());
+		return Graphics == Invalid ? AsyncCompute : Graphics;
+	}
+
+	bool IsSinglePipeline() const
+	{
+		return (Graphics == Invalid) ^ (AsyncCompute == Invalid);
+	}
+
+private:
+	static bool Contains(TInterval<uint32> Interval, uint32 Element)
+	{
+		return Interval.IsValid() && Element > Interval.Min && Element < Interval.Max;
+	}
+
+	static const uint32 Invalid = std::numeric_limits<uint32>::max();
+	uint32 Graphics = Invalid;
+	uint32 AsyncCompute = Invalid;
+	TInterval<uint32> GraphicsForkJoin;
+};
+
+enum class ERHITransientResourceCreateMode
+{
+	// Transient resources are always created inline inside of the Create call.
+	Inline,
+
+	// Transient resource creation may be offloaded to a task (dependent on platform), in which case FRHITransientResource::Finish must be called prior to accessing the underlying RHI resource.
+	Task
+};
+
 class IRHITransientResourceAllocator
 {
 public:
 	virtual ~IRHITransientResourceAllocator() = default;
 
 	// Supports transient allocations of given resource type
-	virtual bool SupportsResourceType(ERHITransientResourceType InType) const = 0;
+	virtual bool SupportsResourceType(ERHITransientResourceType Type) const = 0;
+
+	// Sets the create mode for allocations.
+	virtual void SetCreateMode(ERHITransientResourceCreateMode CreateMode) {};
 
 	// Allocates a new transient resource with memory backed by the transient allocator.
-	virtual FRHITransientTexture* CreateTexture(const FRHITextureCreateInfo& InCreateInfo, const TCHAR* InDebugName, uint32 InPassIndex) = 0;
-	virtual FRHITransientBuffer* CreateBuffer(const FRHIBufferCreateInfo& InCreateInfo, const TCHAR* InDebugName, uint32 InPassIndex) = 0;
+	virtual FRHITransientTexture* CreateTexture(const FRHITextureCreateInfo& CreateInfo, const TCHAR* DebugName, const FRHITransientAllocationFences& Fences) = 0;
+	virtual FRHITransientBuffer* CreateBuffer(const FRHIBufferCreateInfo& CreateInfo, const TCHAR* DebugName, const FRHITransientAllocationFences& Fences) = 0;
 
 	// Deallocates the underlying memory for use by a future resource creation call.
-	virtual void DeallocateMemory(FRHITransientTexture* InTexture, uint32 InPassIndex) = 0;
-	virtual void DeallocateMemory(FRHITransientBuffer* InBuffer, uint32 InPassIndex) = 0;
+	virtual void DeallocateMemory(FRHITransientTexture* Texture, const FRHITransientAllocationFences& Fences) = 0;
+	virtual void DeallocateMemory(FRHITransientBuffer* Buffer, const FRHITransientAllocationFences& Fences) = 0;
 
 	// Flushes any pending allocations prior to rendering. Optionally emits stats if OutStats is valid.
 	virtual void Flush(FRHICommandListImmediate& RHICmdList, FRHITransientAllocationStats* OutStats = nullptr) = 0;
diff --git a/Engine/Source/Runtime/RHI/Public/RHITransition.h b/Engine/Source/Runtime/RHI/Public/RHITransition.h
--- a/Engine/Source/Runtime/RHI/Public/RHITransition.h
+++ b/Engine/Source/Runtime/RHI/Public/RHITransition.h
@@ -302,6 +302,7 @@ struct FRHITransientAliasingInfo
 		return Info;
 	}
 
+	UE_DEPRECATED(5.5, "Discard aliasing ops are no longer necessary.")
 	static FRHITransientAliasingInfo Discard(class FRHITexture* Texture)
 	{
 		FRHITransientAliasingInfo Info;
@@ -311,6 +312,7 @@ struct FRHITransientAliasingInfo
 		return Info;
 	}
 
+	UE_DEPRECATED(5.5, "Discard aliasing ops are no longer necessary.")
 	static FRHITransientAliasingInfo Discard(class FRHIBuffer* Buffer)
 	{
 		FRHITransientAliasingInfo Info;
@@ -325,6 +327,7 @@ struct FRHITransientAliasingInfo
 		return Action == EAction::Acquire;
 	}
 
+	UE_DEPRECATED(5.5, "Discard aliasing ops are no longer necessary.")
 	bool IsDiscard() const
 	{
 		return Action == EAction::Discard;
diff --git a/Engine/Source/Runtime/RHI/Public/RHIValidationCommon.h b/Engine/Source/Runtime/RHI/Public/RHIValidationCommon.h
--- a/Engine/Source/Runtime/RHI/Public/RHIValidationCommon.h
+++ b/Engine/Source/Runtime/RHI/Public/RHIValidationCommon.h
@@ -255,7 +255,7 @@ namespace RHIValidation
 		{}
 
 		void* AcquireBacktrace = nullptr;
-		void* DiscardBacktrace = nullptr;
+		int32 NumAcquiredSubresources = 0;
 
 		bool bTransient = false;
 		EStatus Status = EStatus::None;
@@ -264,7 +264,7 @@ namespace RHIValidation
 		FORCEINLINE bool IsDiscarded() const { return Status == EStatus::Discarded; }
 
 		void Acquire(FResource* Resource, void* CreateTrace);
-		void Discard(FResource* Resource, void* CreateTrace);
+		void Discard(FResource* Resource, void* CreateTrace, ERHIPipeline DiscardPipelines);
 
 		static void AliasingOverlap(FResource* ResourceBefore, FResource* ResourceAfter, void* CreateTrace);
 	};
@@ -325,6 +325,11 @@ namespace RHIValidation
 			return TrackedAccess;
 		}
 
+		inline uint32 GetNumSubresources() const
+		{
+			return NumMips * NumArraySlices * NumPlanes;
+		}
+
 		inline FSubresourceRange GetWholeResourceRange()
 		{
 			checkSlow(NumMips > 0 && NumArraySlices > 0 && NumPlanes > 0);
@@ -352,6 +357,17 @@ namespace RHIValidation
 			check(TransientState.bTransient && TransientState.Status != FTransientState::EStatus::Acquired);
 			TransientState.Status = FTransientState::EStatus::None;
 			DebugName = InDebugName;
+			TrackedAccess = ERHIAccess::Discard;
+
+			for (ERHIPipeline Pipeline : GetRHIPipelines())
+			{
+				auto& State = WholeResourceState.States[Pipeline];
+
+				State.Current.Access = ERHIAccess::Discard;
+				State.Current.Pipelines = Pipeline;
+				State.Previous = State.Current;
+			}
+			SubresourceStates.Reset();
 		}
 
 	protected:
@@ -538,12 +554,6 @@ namespace RHIValidation
 				void* CreateBacktrace;
 			} Data_AcquireTransient;
 
-			struct
-			{
-				FResource* Resource;
-				void* CreateBacktrace;
-			} Data_DiscardTransient;
-
 			struct
 			{
 				FResource* Resource;
@@ -676,17 +686,6 @@ namespace RHIValidation
 			return MoveTemp(Op);
 		}
 
-		static inline FOperation DiscardTransientResource(FResource* Resource, void* CreateBacktrace)
-		{
-			Resource->AddOpRef();
-
-			FOperation Op;
-			Op.Type = EOpType::DiscardTransient;
-			Op.Data_DiscardTransient.Resource = Resource;
-			Op.Data_DiscardTransient.CreateBacktrace = CreateBacktrace;
-			return MoveTemp(Op);
-		}
-
 		static inline FOperation InitTransient(FResource* Resource, const TCHAR* DebugName)
 		{
 			Resource->AddOpRef();
diff --git a/Engine/Source/Runtime/RHI/Public/RHIValidationTransientResourceAllocator.h b/Engine/Source/Runtime/RHI/Public/RHIValidationTransientResourceAllocator.h
--- a/Engine/Source/Runtime/RHI/Public/RHIValidationTransientResourceAllocator.h
+++ b/Engine/Source/Runtime/RHI/Public/RHIValidationTransientResourceAllocator.h
@@ -16,11 +16,12 @@ public:
 	virtual ~FValidationTransientResourceAllocator();
 
 	// Implementation of FRHITransientResourceAllocator interface
+	virtual void SetCreateMode(ERHITransientResourceCreateMode InCreateMode) override final;
 	virtual bool SupportsResourceType(ERHITransientResourceType InType) const override final { return RHIAllocator->SupportsResourceType(InType); }
-	virtual FRHITransientTexture* CreateTexture(const FRHITextureCreateInfo& InCreateInfo, const TCHAR* InDebugName, uint32 InPassIndex) override final;
-	virtual FRHITransientBuffer* CreateBuffer(const FRHIBufferCreateInfo& InCreateInfo, const TCHAR* InDebugName, uint32 InPassIndex) override final;
-	virtual void DeallocateMemory(FRHITransientTexture* InTexture, uint32 InPassIndex) override final;
-	virtual void DeallocateMemory(FRHITransientBuffer* InBuffer, uint32 InPassIndex) override final;
+	virtual FRHITransientTexture* CreateTexture(const FRHITextureCreateInfo& InCreateInfo, const TCHAR* InDebugName, const FRHITransientAllocationFences& Fences) override final;
+	virtual FRHITransientBuffer* CreateBuffer(const FRHIBufferCreateInfo& InCreateInfo, const TCHAR* InDebugName, const FRHITransientAllocationFences& Fences) override final;
+	virtual void DeallocateMemory(FRHITransientTexture* InTexture, const FRHITransientAllocationFences& Fences) override final;
+	virtual void DeallocateMemory(FRHITransientBuffer* InBuffer, const FRHITransientAllocationFences& Fences) override final;
 	virtual void Flush(FRHICommandListImmediate&, FRHITransientAllocationStats*) override final;
 	virtual void Release(FRHICommandListImmediate&) override final;
 
diff --git a/Engine/Source/Runtime/RHICore/Private/RHICoreTransientResourceAllocator.cpp b/Engine/Source/Runtime/RHICore/Private/RHICoreTransientResourceAllocator.cpp
--- a/Engine/Source/Runtime/RHICore/Private/RHICoreTransientResourceAllocator.cpp
+++ b/Engine/Source/Runtime/RHICore/Private/RHICoreTransientResourceAllocator.cpp
@@ -7,6 +7,21 @@
 #include "ProfilingDebugging/CountersTrace.h"
 #include "ProfilingDebugging/CsvProfiler.h"
 #include "RHICommandList.h"
+#include "RHICoreNvidiaAftermath.h"
+
+static bool GRHITransientAllocatorSeparateAsyncComputeHeap = false;
+static FAutoConsoleVariableRef CVarRHITransientAllocatorSeparateAsyncComputeHeap(
+	TEXT("RHI.TransientAllocator.SeparateAsyncComputeHeap"),
+	GRHITransientAllocatorSeparateAsyncComputeHeap,
+	TEXT("If enabled, async compute only allocations are put in their own separate heap."),
+	ECVF_RenderThreadSafe);
+
+static bool GRHITransientAllocatorParallelResourceCreation = true;
+static FAutoConsoleVariableRef CVarRHITransientAllocatorParallelResourceCreation(
+	TEXT("RHI.TransientAllocator.ParallelResourceCreation"),
+	GRHITransientAllocatorParallelResourceCreation,
+	TEXT("If enabled, a task is launched for each placed resource that is created."),
+	ECVF_RenderThreadSafe);
 
 static int32 GRHITransientAllocatorMinimumHeapSize = 128;
 static FAutoConsoleVariableRef CVarRHITransientAllocatorMinimumHeapSize(
@@ -104,169 +119,151 @@ void FRHITransientMemoryStats::Submit(uint64 UsedSize)
 	Reset();
 }
 
+//////////////////////////////////////////////////////////////////////////
+// Transient Resource Heap Allocator
 //////////////////////////////////////////////////////////////////////////
 
-void FRHITransientResourceOverlapTracker::Track(FRHITransientResource* TransientResource, uint32 PageOffsetMin, uint32 PageOffsetMax)
+FRHITransientHeapAllocator::FRHITransientHeapAllocator(uint64 InCapacity, uint32 InAlignment)
+	: Capacity(InCapacity)
+	, AlignmentMin(InAlignment)
 {
-	check(TransientResource);
+	HeadHandle = CreateRange();
+	InsertRange(HeadHandle, nullptr, {}, 0, Capacity);
+}
 
-	FResourceRange ResourceRangeNew;
-	ResourceRangeNew.ResourceIndex = Resources.Emplace(TransientResource);
-	ResourceRangeNew.PageOffsetMin = PageOffsetMin;
-	ResourceRangeNew.PageOffsetMax = PageOffsetMax;
+FRHITransientHeapAllocation FRHITransientHeapAllocator::Allocate(const FRHITransientAllocationFences& Fences, uint64 Size, uint32 Alignment, TArray<FAliasingOverlap>& OutAliasingOverlaps)
+{
+	check(Size > 0);
 
-	int32 InsertIndex = Algo::LowerBound(ResourceRanges, ResourceRangeNew, [](const FResourceRange& Lhs, const FResourceRange& Rhs)
+	if (Alignment < AlignmentMin)
 	{
-		return Lhs.PageOffsetMax <= Rhs.PageOffsetMin;
-	});
+		Alignment = AlignmentMin;
+	}
+
+	TArray<FRangeHandle, TInlineAllocator<64>> RangeCandidates;
 
-	for (int32 Index = InsertIndex; Index < ResourceRanges.Num(); ++Index)
+	FRangeHandle Handle = GetFirstFreeRangeHandle();
+	FRangeHandle FirstPreviousHandle = HeadHandle;
+	FRangeHandle PreviousHandle = InvalidRangeHandle;
+	uint64 FirstAllocationRegionMin = 0;
+	uint64 AllocationMin = 0;
+	uint64 AllocationMax = 0;
+	uint64 NextRangeMin  = 0;
+	uint64 LeftoverSize  = 0;
+	bool bAllocationComplete = false;
+
+	while (Handle != InvalidRangeHandle)
 	{
-		FResourceRange& ResourceRangeOld = ResourceRanges[Index];
+		FRange& Range = Ranges[Handle];
+		const uint64 RangeMax = Range.Offset + Range.Size;
 
-		// If the old range starts later in the heap and doesn't overlap, the sort invariant guarantees no future range will overlap.
-		if (ResourceRangeOld.PageOffsetMin >= ResourceRangeNew.PageOffsetMax)
+		const auto NextRegion = [&]
 		{
-			break;
-		}
-
-		TransientResource->AddAliasingOverlap(Resources[ResourceRangeOld.ResourceIndex]);
+			PreviousHandle = Handle;
+			Handle = Range.NextFreeHandle;
+		};
 
-		// Complete overlap.
-		if (ResourceRangeOld.PageOffsetMin >= ResourceRangeNew.PageOffsetMin && ResourceRangeOld.PageOffsetMax <= ResourceRangeNew.PageOffsetMax)
+		// Specify the initial min / max bounds based off the current candidate range.
+		if (RangeCandidates.IsEmpty())
 		{
-			ResourceRanges.RemoveAt(Index, EAllowShrinking::No);
-			Index--;
+			const uint64 AlignedOffset = Align(GpuVirtualAddress + Range.Offset, Alignment) - GpuVirtualAddress;
+
+			// Skip regions smaller than the alignment padding.
+			if (AlignedOffset >= RangeMax)
+			{
+				FirstPreviousHandle = InvalidRangeHandle;
+				NextRegion();
+				continue;
+			}
+
+			FirstAllocationRegionMin = NextRangeMin = Range.Offset;
+			AllocationMin = AlignedOffset;
+			AllocationMax = AlignedOffset + Size;
 		}
-		// Partial overlap, can manifest as three cases:
-		else
+
+		// Range is allowed to be part of this allocation.
+		if (Range.Offset == NextRangeMin)
 		{
-			// 1) New:    ********
-			//            |||        ->
-			//    Old: ======             ===********
-			if (ResourceRangeOld.PageOffsetMin < ResourceRangeNew.PageOffsetMin)
+			ON_SCOPE_EXIT { NextRegion(); };
+
+			if (!FRHITransientAllocationFences::Contains(Range.Fences, Fences))
 			{
-				FResourceRange ResourceRangeOldCopy = ResourceRangeOld;
-				ResourceRangeOld.PageOffsetMax = ResourceRangeNew.PageOffsetMin;
-
-				InsertIndex++;
-				// 3) New:    ********
-				//            ||||||||      ->
-				//    Old: ==============        ===********===
-				if (ResourceRangeOldCopy.PageOffsetMax > ResourceRangeNew.PageOffsetMax)
+				if (FirstPreviousHandle == InvalidRangeHandle)
 				{
-					ResourceRangeOldCopy.PageOffsetMin = ResourceRangeNew.PageOffsetMax;
+					FirstPreviousHandle = PreviousHandle;
+				}
 
-					// Lower bound has been resized already; add an upper bound.
-					FResourceRange ResourceRangeOldUpper = ResourceRangeOldCopy;
-					ResourceRanges.Insert(ResourceRangeOldUpper, Index + 1);
+				RangeCandidates.Emplace(Handle);
 
+				// Range is large enough to service remaining allocation
+				if (AllocationMax <= RangeMax)
+				{
+					LeftoverSize = RangeMax - AllocationMax;
+					bAllocationComplete = true;
 					break;
 				}
-			}
-			else
-			{
-				// 2) New:    ********
-				//                |||      ->
-				//    Old:        ======         ********===
-				ResourceRangeOld.PageOffsetMin = ResourceRangeNew.PageOffsetMax;
 
-				break;
+				NextRangeMin += Range.Size;
+				continue;
 			}
 		}
-	}
-
-	ResourceRanges.Insert(ResourceRangeNew, InsertIndex);
-}
-
-//////////////////////////////////////////////////////////////////////////
-
-void FRHITransientResourceOverlapTracker::Reset()
-{
-	TArray<FResourceRange> AcquiredRanges;
-	TArray<FRHITransientResource*> AcquiredResources;
-	AcquiredRanges.Reserve(ResourceRanges.Num());
-	AcquiredResources.Reserve(Resources.Num());
 
-	for (FResourceRange Range : ResourceRanges)
-	{
-		FRHITransientResource* Resource = Resources[Range.ResourceIndex];
-
-		if (Resource->IsAcquired())
-		{
-			Range.ResourceIndex = AcquiredResources.Num();
-			AcquiredRanges.Emplace(Range);
-			AcquiredResources.Emplace(Resource);
-		}
+		RangeCandidates.Reset();
+		FirstPreviousHandle = InvalidRangeHandle;
 	}
 
-	Swap(ResourceRanges, AcquiredRanges);
-	Swap(Resources, AcquiredResources);
-}
-
-//////////////////////////////////////////////////////////////////////////
-// Transient Resource Heap Allocator
-//////////////////////////////////////////////////////////////////////////
-
-FRHITransientHeapAllocator::FRHITransientHeapAllocator(uint64 InCapacity, uint32 InAlignment)
-	: Capacity(InCapacity)
-	, AlignmentMin(InAlignment)
-{
-	HeadHandle = CreateRange();
-	InsertRange(HeadHandle, 0, Capacity);
-}
-
-FRHITransientHeapAllocation FRHITransientHeapAllocator::Allocate(uint64 Size, uint32 Alignment)
-{
-	check(Size > 0);
+	FRHITransientHeapAllocation Allocation;
 
-	if (Alignment < AlignmentMin)
+	if (bAllocationComplete)
 	{
-		Alignment = AlignmentMin;
-	}
+		check(!RangeCandidates.IsEmpty())
+		const uint64 AlignedSize  = AllocationMax - FirstAllocationRegionMin;
+		const uint64 AlignmentPad = AlignedSize - Size;
 
-	FFindResult FindResult = FindFreeRange(Size, Alignment);
+		AllocationCount++;
+		UsedSize       += AlignedSize;
+		AlignmentWaste += AlignmentPad;
 
-	if (FindResult.FoundHandle == InvalidRangeHandle)
-	{
-		return {};
-	}
-
-	FRange& FoundRange = Ranges[FindResult.FoundHandle];
+		for (int32 Index = 0; Index < RangeCandidates.Num(); ++Index)
+		{
+			int32 RangeIndex = RangeCandidates[Index];
+			const FRange& Range = Ranges[RangeIndex];
 
-	const uint64 AlignedSize = FoundRange.Size - FindResult.LeftoverSize;
-	const uint64 AlignmentPad = AlignedSize - Size;
-	const uint64 AlignedOffset = FoundRange.Offset + AlignmentPad;
-	const uint64 AllocationEnd = AlignedOffset + Size;
+			if (FRHITransientResource* ResourceToOverlap = Range.Resource)
+			{
+				OutAliasingOverlaps.Emplace(ResourceToOverlap, FRHITransientAllocationFences::GetAcquireFence(Range.Fences, Fences));
+			}
 
-	// Adjust the range if there is space left over.
-	if (FindResult.LeftoverSize)
-	{
-		FoundRange.Offset = AllocationEnd;
-		FoundRange.Size = FindResult.LeftoverSize;
-	}
-	// Otherwise, remove it.
-	else
-	{
-		RemoveRange(FindResult.PreviousHandle, FindResult.FoundHandle);
+			if (Index < RangeCandidates.Num() - 1)
+			{
+				RemoveRange(FirstPreviousHandle, RangeIndex);
+			}
+		}
+	
+		if (LeftoverSize > 0)
+		{
+			FRange& LastRange = Ranges[RangeCandidates.Last()];
+			LastRange.Offset  = AllocationMax;
+			LastRange.Size    = LeftoverSize;
+		}
+		else
+		{
+			RemoveRange(FirstPreviousHandle, RangeCandidates.Last());
+		}
+		
+		Allocation.Size   = Size;
+		Allocation.Offset = AllocationMin;
+		Allocation.AlignmentPad = AlignmentPad;
 	}
 
-	AllocationCount++;
-	UsedSize += AlignedSize;
-	AlignmentWaste += AlignmentPad;
-
-	FRHITransientHeapAllocation Allocation;
-	Allocation.Size = Size;
-	Allocation.Offset = AlignedOffset;
-	Allocation.AlignmentPad = AlignmentPad;
-
 	Validate();
-
 	return Allocation;
 }
 
-void FRHITransientHeapAllocator::Deallocate(FRHITransientHeapAllocation Allocation)
+void FRHITransientHeapAllocator::Deallocate(FRHITransientResource* Resource, const FRHITransientAllocationFences& Fences)
 {
+	check(Resource);
+	const FRHITransientHeapAllocation& Allocation = Resource->GetHeapAllocation();
 	check(Allocation.Size > 0 && Allocation.Size <= UsedSize);
 
 	// Reconstruct the original range offset by subtracting the alignment pad, and expand the size accordingly.
@@ -275,7 +272,6 @@ void FRHITransientHeapAllocator::Deallocate(FRHITransientHeapAllocation Allocati
 	const uint64 RangeToFreeEnd = RangeToFreeOffset + RangeToFreeSize;
 
 	FRangeHandle PreviousHandle = HeadHandle;
-	FRangeHandle NextHandle = InvalidRangeHandle;
 	FRangeHandle Handle = GetFirstFreeRangeHandle();
 
 	while (Handle != InvalidRangeHandle)
@@ -285,7 +281,6 @@ void FRHITransientHeapAllocator::Deallocate(FRHITransientHeapAllocation Allocati
 		// Find the first free range after the one being freed.
 		if (RangeToFreeOffset < Range.Offset)
 		{
-			NextHandle = Handle;
 			break;
 		}
 
@@ -293,88 +288,41 @@ void FRHITransientHeapAllocator::Deallocate(FRHITransientHeapAllocation Allocati
 		Handle = Range.NextFreeHandle;
 	}
 
-	uint64 MergedFreeRangeStart = RangeToFreeOffset;
-	uint64 MergedFreeRangeEnd = RangeToFreeEnd;
-	bool bMergedPrevious = false;
-	bool bMergedNext = false;
-
-	if (PreviousHandle != HeadHandle)
-	{
-		FRange& PreviousRange = Ranges[PreviousHandle];
-
-		// Attempt to merge the previous range with the range being freed.
-		if (PreviousRange.GetEnd() == RangeToFreeOffset)
-		{
-			PreviousRange.Size += RangeToFreeSize;
-			MergedFreeRangeStart = PreviousRange.Offset;
-			MergedFreeRangeEnd = PreviousRange.GetEnd();
-			bMergedPrevious = true;
-		}
-	}
-
-	if (NextHandle != InvalidRangeHandle)
-	{
-		FRange& NextRange = Ranges[NextHandle];
+	InsertRange(PreviousHandle, Resource, Fences, RangeToFreeOffset, RangeToFreeSize);
 
-		// Attempt to merge the next range with the range being freed.
-		if (RangeToFreeEnd == NextRange.Offset)
-		{
-			NextRange.Size += RangeToFreeSize;
-			NextRange.Offset = RangeToFreeOffset;
-			MergedFreeRangeStart = FMath::Min(MergedFreeRangeStart, RangeToFreeOffset);
-			MergedFreeRangeEnd = NextRange.GetEnd();
-			bMergedNext = true;
-		}
-	}
-
-	// With both previous and next ranges merged with the freed range, they now overlap. Remove next and expand previous to cover all three.
-	if (bMergedPrevious && bMergedNext)
-	{
-		FRange& PreviousRange = Ranges[PreviousHandle];
-		FRange& NextRange = Ranges[NextHandle];
-
-		PreviousRange.Size = MergedFreeRangeEnd - MergedFreeRangeStart;
-		RemoveRange(PreviousHandle, NextHandle);
-	}
-	// If neither previous or next were merged, insert a new range between them.
-	else if (!bMergedPrevious && !bMergedNext)
-	{
-		InsertRange(PreviousHandle, RangeToFreeOffset, RangeToFreeSize);
-	}
-
-	UsedSize -= RangeToFreeSize;
+	UsedSize       -= RangeToFreeSize;
 	AlignmentWaste -= Allocation.AlignmentPad;
 	AllocationCount--;
 
 	Validate();
 }
 
-FRHITransientHeapAllocator::FFindResult FRHITransientHeapAllocator::FindFreeRange(uint64 Size, uint32 Alignment)
+void FRHITransientHeapAllocator::Flush()
 {
-	FFindResult FindResult;
-	FindResult.PreviousHandle = HeadHandle;
-
 	FRangeHandle Handle = GetFirstFreeRangeHandle();
+	FRangeHandle PreviousHandle = InvalidRangeHandle;
+
 	while (Handle != InvalidRangeHandle)
 	{
 		FRange& Range = Ranges[Handle];
+		Range.Fences = {};
+		Range.Resource = nullptr;
 
-		// Due to alignment we may have to shift the offset and expand the size accordingly.
-		const uint64 AlignmentPad = Align(GpuVirtualAddress + Range.Offset, Alignment) - GpuVirtualAddress - Range.Offset;
-		const uint64 RequiredSize = Size + AlignmentPad;
-
-		if (RequiredSize <= Range.Size)
+		if (PreviousHandle != InvalidRangeHandle)
 		{
-			FindResult.FoundHandle = Handle;
-			FindResult.LeftoverSize = Range.Size - RequiredSize;
-			return FindResult;
+			FRange& PreviousRange = Ranges[PreviousHandle];
+
+			if (PreviousRange.Offset + PreviousRange.Size == Range.Offset)
+			{
+				PreviousRange.Size += Range.Size;
+				Handle = RemoveRange(PreviousHandle, Handle);
+				continue;
+			}
 		}
 
-		FindResult.PreviousHandle = Handle;
 		Handle = Range.NextFreeHandle;
+		PreviousHandle = Handle;
 	}
-
-	return {};
 }
 
 void FRHITransientHeapAllocator::Validate()
@@ -396,7 +344,7 @@ void FRHITransientHeapAllocator::Validate()
 			const FRange& PreviousRange = Ranges[PreviousHandle];
 
 			// Checks that the ranges are sorted.
-			check(PreviousRange.Offset + PreviousRange.Size < Range.Offset);
+			check(PreviousRange.Offset + PreviousRange.Size <= Range.Offset);
 		}
 
 		PreviousHandle = Handle;
@@ -412,13 +360,13 @@ void FRHITransientHeapAllocator::Validate()
 FRHITransientTexture* FRHITransientHeap::CreateTexture(
 	const FRHITextureCreateInfo& CreateInfo,
 	const TCHAR* DebugName,
-	uint32 PassIndex,
+	const FRHITransientAllocationFences& Fences,
 	uint64 CurrentAllocatorCycle,
 	uint64 TextureSize,
 	uint32 TextureAlignment,
 	FCreateTextureFunction CreateTextureFunction)
 {
-	FRHITransientHeapAllocation Allocation = Allocator.Allocate(TextureSize, TextureAlignment);
+	FRHITransientHeapAllocation Allocation = Allocator.Allocate(Fences, TextureSize, TextureAlignment, AliasingOverlaps);
 	Allocation.Heap = this;
 
 	if (!Allocation.IsValid())
@@ -433,28 +381,28 @@ FRHITransientTexture* FRHITransientHeap::CreateTexture(
 	});
 
 	check(Texture);
-	Texture->Acquire(FRHICommandListImmediate::Get(), DebugName, PassIndex, CurrentAllocatorCycle);
-	AllocateMemoryInternal(Texture, DebugName, PassIndex, CurrentAllocatorCycle, Allocation);
+	Texture->Acquire(DebugName, Fences.GetSinglePipeline(), CurrentAllocatorCycle);
+	AllocateMemoryInternal(Texture, Allocation);
 	Stats.AllocateTexture(Allocation.Size);
 	return Texture;
 }
 
-void FRHITransientHeap::DeallocateMemory(FRHITransientTexture* Texture, uint32 PassIndex)
+void FRHITransientHeap::DeallocateMemory(FRHITransientTexture* Texture, const FRHITransientAllocationFences& Fences)
 {
-	DeallocateMemoryInternal(Texture, PassIndex);
+	DeallocateMemoryInternal(Texture, Fences);
 	Stats.DeallocateTexture(Texture->GetSize());
 }
 
 FRHITransientBuffer* FRHITransientHeap::CreateBuffer(
 	const FRHIBufferCreateInfo& CreateInfo,
 	const TCHAR* DebugName,
-	uint32 PassIndex,
+	const FRHITransientAllocationFences& Fences,
 	uint64 CurrentAllocatorCycle,
 	uint64 BufferSize,
 	uint32 BufferAlignment,
 	FCreateBufferFunction CreateBufferFunction)
 {
-	FRHITransientHeapAllocation Allocation = Allocator.Allocate(BufferSize, BufferAlignment);
+	FRHITransientHeapAllocation Allocation = Allocator.Allocate(Fences, BufferSize, BufferAlignment, AliasingOverlaps);
 	Allocation.Heap = this;
 
 	if (!Allocation.IsValid())
@@ -469,39 +417,35 @@ FRHITransientBuffer* FRHITransientHeap::CreateBuffer(
 	});
 
 	check(Buffer);
-	Buffer->Acquire(FRHICommandListImmediate::Get(), DebugName, PassIndex, CurrentAllocatorCycle);
-	AllocateMemoryInternal(Buffer, DebugName, PassIndex, CurrentAllocatorCycle, Allocation);
+	Buffer->Acquire(DebugName, Fences.GetSinglePipeline(), CurrentAllocatorCycle);
+	AllocateMemoryInternal(Buffer, Allocation);
 	Stats.AllocateBuffer(Allocation.Size);
 	return Buffer;
 }
 
-void FRHITransientHeap::DeallocateMemory(FRHITransientBuffer* Buffer, uint32 PassIndex)
+void FRHITransientHeap::DeallocateMemory(FRHITransientBuffer* Buffer, const FRHITransientAllocationFences& Fences)
 {
-	DeallocateMemoryInternal(Buffer, PassIndex);
+	DeallocateMemoryInternal(Buffer, Fences);
 	Stats.DeallocateBuffer(Buffer->GetSize());
 }
 
-void FRHITransientHeap::AllocateMemoryInternal(FRHITransientResource* Resource, const TCHAR* Name, uint32 PassIndex, uint64 CurrentAllocatorCycle, const FRHITransientHeapAllocation& Allocation)
+void FRHITransientHeap::AllocateMemoryInternal(FRHITransientResource* Resource, const FRHITransientHeapAllocation& Allocation)
 {
 	Resource->GetHeapAllocation() = Allocation;
 
-	check(Allocation.Offset % Initializer.Alignment == 0);
-		const uint64 AlignedSize   = Align(Allocation.Size, Initializer.Alignment);
-		const uint32 PageOffsetMin = Allocation.Offset >> AlignmentLog2;
-		const uint32 PageOffsetMax = (Allocation.Offset + AlignedSize) >> AlignmentLog2;
-		OverlapTracker.Track(Resource, PageOffsetMin, PageOffsetMax);
+	for (const FRHITransientHeapAllocator::FAliasingOverlap& AliasingOverlap : AliasingOverlaps)
+	{
+		Resource->AddAliasingOverlap(AliasingOverlap.Resource, AliasingOverlap.AcquireFence);
+	}
+	AliasingOverlaps.Reset();
 
 	CommitSize = FMath::Max(CommitSize, Allocation.Offset + Allocation.Size);
 }
 
-void FRHITransientHeap::DeallocateMemoryInternal(FRHITransientResource* Resource, uint32 PassIndex)
+void FRHITransientHeap::DeallocateMemoryInternal(FRHITransientResource* Resource, const FRHITransientAllocationFences& Fences)
 {
-	Resource->Discard(PassIndex);
-
-	const FRHITransientHeapAllocation Allocation = Resource->GetHeapAllocation();
-	check(Allocation.Heap == this);
-
-	Allocator.Deallocate(Allocation);
+	Resource->Discard();
+	Allocator.Deallocate(Resource, Fences);
 }
 
 void FRHITransientHeap::Flush(uint64 AllocatorCycle, FRHITransientMemoryStats& OutMemoryStats, FRHITransientAllocationStats* OutAllocationStats)
@@ -510,6 +454,8 @@ void FRHITransientHeap::Flush(uint64 AllocatorCycle, FRHITransientMemoryStats& O
 	OutMemoryStats.Accumulate(Stats);
 	Stats.Reset();
 
+	Allocator.Flush();
+
 	if (OutAllocationStats)
 	{
 		const auto AddResourceToStats = [](FRHITransientAllocationStats& AllocationStats, FRHITransientResource* Resource)
@@ -542,6 +488,8 @@ void FRHITransientHeap::Flush(uint64 AllocatorCycle, FRHITransientMemoryStats& O
 		OutAllocationStats->MemoryRanges.Add(MemoryRange);
 	}
 
+	CommitSizeMax = FMath::Max(CommitSize, CommitSizeMax);
+
 	if (bHasDeallocations)
 	{
 		CommitSize = 0;
@@ -565,9 +513,7 @@ void FRHITransientHeap::Flush(uint64 AllocatorCycle, FRHITransientMemoryStats& O
 				CommitSize = FMath::Max(CommitSize, Allocation.Offset + Allocation.Size);
 			}
 		}
-
-			OverlapTracker.Reset();
-		}
+	}
 }
 
 //////////////////////////////////////////////////////////////////////////
@@ -619,7 +565,6 @@ FRHITransientHeap* FRHITransientHeapCache::Acquire(uint64 FirstAllocationSize, E
 	FRHITransientHeap* Heap = CreateHeap(HeapInitializer);
 	check(Heap);
 
-	TotalMemoryCapacity += HeapInitializer.Size;
 	LiveList.Emplace(Heap);
 	return Heap;
 }
@@ -642,13 +587,14 @@ void FRHITransientHeapCache::GarbageCollect()
 {
 	FScopeLock Lock(&CriticalSection);
 
+	uint64 TotalCommitSize = 0;
+
 	for (int32 HeapIndex = 0; HeapIndex < FreeList.Num(); ++HeapIndex)
 	{
 		FRHITransientHeap* Heap = FreeList[HeapIndex];
 
 		if (Heap->GetLastUsedGarbageCollectCycle() + Initializer.GarbageCollectLatency <= GarbageCollectCycle)
 		{
-			TotalMemoryCapacity -= Heap->GetCapacity();
 			FreeList.RemoveAt(HeapIndex);
 			LiveList.Remove(Heap);
 			HeapIndex--;
@@ -656,10 +602,16 @@ void FRHITransientHeapCache::GarbageCollect()
 			delete Heap;
 		}
 	}
+	
+	for (FRHITransientHeap* Heap : LiveList)
+	{
+		TotalCommitSize += Initializer.bSupportsVirtualMapping ? Heap->CommitSizeMax : Heap->GetCapacity();
+		Heap->CommitSizeMax = 0;
+	}
 
 	TRACE_COUNTER_SET(TransientMemoryRangeCount, LiveList.Num());
 
-	Stats.Submit(TotalMemoryCapacity);
+	Stats.Submit(TotalCommitSize);
 
 	GarbageCollectCycle++;
 }
@@ -669,7 +621,7 @@ void FRHITransientHeapCache::GarbageCollect()
 FRHITransientTexture* FRHITransientResourceHeapAllocator::CreateTextureInternal(
 	const FRHITextureCreateInfo& CreateInfo,
 	const TCHAR* DebugName,
-	uint32 PassIndex,
+	const FRHITransientAllocationFences& Fences,
 	uint64 TextureSize,
 	uint32 TextureAlignment,
 	FRHITransientHeap::FCreateTextureFunction CreateTextureFunction)
@@ -688,7 +640,7 @@ FRHITransientTexture* FRHITransientResourceHeapAllocator::CreateTextureInternal(
 			continue;
 		}
 
-		Texture = Heap->CreateTexture(CreateInfo, DebugName, PassIndex, CurrentCycle, TextureSize, TextureAlignment, CreateTextureFunction);
+		Texture = Heap->CreateTexture(CreateInfo, DebugName, Fences, CurrentCycle, TextureSize, TextureAlignment, CreateTextureFunction);
 
 		if (Texture)
 		{
@@ -701,9 +653,9 @@ FRHITransientTexture* FRHITransientResourceHeapAllocator::CreateTextureInternal(
 		FRHITransientHeap* Heap = HeapCache.Acquire(TextureSize, TextureHeapFlags);
 		Heaps.Emplace(Heap);
 
-		Texture = Heap->CreateTexture(CreateInfo, DebugName, PassIndex, CurrentCycle, TextureSize, TextureAlignment, CreateTextureFunction);
+		Texture = Heap->CreateTexture(CreateInfo, DebugName, Fences, CurrentCycle, TextureSize, TextureAlignment, CreateTextureFunction);
 	}
-
+	
 	check(Texture);
 	IF_RHICORE_TRANSIENT_ALLOCATOR_DEBUG(ActiveResources.Emplace(Texture));
 	return Texture;
@@ -712,7 +664,7 @@ FRHITransientTexture* FRHITransientResourceHeapAllocator::CreateTextureInternal(
 FRHITransientBuffer* FRHITransientResourceHeapAllocator::CreateBufferInternal(
 	const FRHIBufferCreateInfo& CreateInfo,
 	const TCHAR* DebugName,
-	uint32 PassIndex,
+	const FRHITransientAllocationFences& Fences,
 	uint32 BufferSize,
 	uint32 BufferAlignment,
 	FRHITransientHeap::FCreateBufferFunction CreateBufferFunction)
@@ -726,7 +678,7 @@ FRHITransientBuffer* FRHITransientResourceHeapAllocator::CreateBufferInternal(
 			continue;
 		}
 
-		Buffer = Heap->CreateBuffer(CreateInfo, DebugName, PassIndex, CurrentCycle, BufferSize, BufferAlignment, CreateBufferFunction);
+		Buffer = Heap->CreateBuffer(CreateInfo, DebugName, Fences, CurrentCycle, BufferSize, BufferAlignment, CreateBufferFunction);
 
 		if (Buffer)
 		{
@@ -739,15 +691,15 @@ FRHITransientBuffer* FRHITransientResourceHeapAllocator::CreateBufferInternal(
 		FRHITransientHeap* Heap = HeapCache.Acquire(BufferSize, ERHITransientHeapFlags::AllowBuffers);
 		Heaps.Emplace(Heap);
 
-		Buffer = Heap->CreateBuffer(CreateInfo, DebugName, PassIndex, CurrentCycle, BufferSize, BufferAlignment, CreateBufferFunction);
+		Buffer = Heap->CreateBuffer(CreateInfo, DebugName, Fences, CurrentCycle, BufferSize, BufferAlignment, CreateBufferFunction);
 	}
-
+	
 	check(Buffer);
 	IF_RHICORE_TRANSIENT_ALLOCATOR_DEBUG(ActiveResources.Emplace(Buffer));
 	return Buffer;
 }
 
-void FRHITransientResourceHeapAllocator::DeallocateMemory(FRHITransientTexture* Texture, uint32 PassIndex)
+void FRHITransientResourceHeapAllocator::DeallocateMemory(FRHITransientTexture* Texture, const FRHITransientAllocationFences& Fences)
 {
 	check(Texture);
 
@@ -756,13 +708,13 @@ void FRHITransientResourceHeapAllocator::DeallocateMemory(FRHITransientTexture*
 	check(Heap);
 	check(Heaps.Contains(Heap));
 
-	Heap->DeallocateMemory(Texture, PassIndex);
+	Heap->DeallocateMemory(Texture, Fences);
 	DeallocationCount++;
 
 	IF_RHICORE_TRANSIENT_ALLOCATOR_DEBUG(ActiveResources.Remove(Texture));
 }
 
-void FRHITransientResourceHeapAllocator::DeallocateMemory(FRHITransientBuffer* Buffer, uint32 PassIndex)
+void FRHITransientResourceHeapAllocator::DeallocateMemory(FRHITransientBuffer* Buffer, const FRHITransientAllocationFences& Fences)
 {
 	check(Buffer);
 
@@ -771,12 +723,24 @@ void FRHITransientResourceHeapAllocator::DeallocateMemory(FRHITransientBuffer* B
 	check(Heap);
 	check(Heaps.Contains(Heap));
 
-	Heap->DeallocateMemory(Buffer, PassIndex);
+	Heap->DeallocateMemory(Buffer, Fences);
 	DeallocationCount++;
 
 	IF_RHICORE_TRANSIENT_ALLOCATOR_DEBUG(ActiveResources.Remove(Buffer));
 }
 
+void FRHITransientResourceHeapAllocator::SetCreateMode(ERHITransientResourceCreateMode InCreateMode)
+{
+	const bool bSupportsParallelResourceCreation = GRHITransientAllocatorParallelResourceCreation
+#if NV_AFTERMATH
+		// Aftermath adds locks that serialize placed resource creation.
+		&& !UE::RHICore::Nvidia::Aftermath::IsEnabled()
+#endif
+		;
+
+	CreateMode = bSupportsParallelResourceCreation ? InCreateMode : ERHITransientResourceCreateMode::Inline;
+}
+
 void FRHITransientResourceHeapAllocator::Flush(FRHICommandListImmediate& RHICmdList, FRHITransientAllocationStats* OutAllocationStats)
 {
 	FRHITransientMemoryStats Stats;
@@ -821,7 +785,7 @@ void FRHITransientResourceHeapAllocator::Flush(FRHICommandListImmediate& RHICmdL
 		DeallocationCount = 0;
 	}
 
-	RHICmdList.EnqueueLambda([&HeapCache = HeapCache, Stats](FRHICommandListImmediate&)
+	RHICmdList.EnqueueLambda([&HeapCache = HeapCache, Stats](FRHICommandListBase&)
 	{
 		HeapCache.Stats.Accumulate(Stats);
 	});
@@ -871,18 +835,22 @@ void FRHITransientPageSpanAllocator::Reset()
 		// Allocate head and tail nodes (0 and 1)
 		for (uint32 Index = 0; Index < 2; Index++)
 		{
+			PageSpans[Index].Resource = nullptr;
 			PageSpans[Index].Offset = 0;
 			PageSpans[Index].Count = 0;
-			PageSpans[Index].PrevSpanIndex = InvalidIndex;
+			PageSpans[Index].Fences = {};
 			PageSpans[Index].NextSpanIndex = InvalidIndex;
+			PageSpans[Index].PrevSpanIndex = InvalidIndex;
 			PageSpans[Index].bAllocated = false;
 		}
 		PageSpans[HeadSpanIndex].NextSpanIndex = FirstFreeNodeIndex;
 		PageSpans[TailSpanIndex].PrevSpanIndex = FirstFreeNodeIndex;
 
 		// First Node
+		PageSpans[FirstFreeNodeIndex].Resource = nullptr;
 		PageSpans[FirstFreeNodeIndex].Offset = 0;
 		PageSpans[FirstFreeNodeIndex].Count = MaxPageCount;
+		PageSpans[FirstFreeNodeIndex].Fences = {};
 		PageSpans[FirstFreeNodeIndex].PrevSpanIndex = HeadSpanIndex;
 		PageSpans[FirstFreeNodeIndex].NextSpanIndex = TailSpanIndex;
 		PageSpans[FirstFreeNodeIndex].bAllocated = false;
@@ -903,95 +871,101 @@ void FRHITransientPageSpanAllocator::Reset()
 	}
 }
 
-bool FRHITransientPageSpanAllocator::Allocate(uint32 PageCount, uint32& OutNumPagesAllocated, uint32& OutSpanIndex)
+bool FRHITransientPageSpanAllocator::Allocate(FRHITransientResource* Resource, const FRHITransientAllocationFences& Fences, uint32 PageCount, uint32& OutNumPagesAllocated, uint32& OutSpanIndex)
 {
 	OutNumPagesAllocated = 0;
 
-	if (FreePageCount < PageCount)
-	{
-		// If we're allowing partial allocs and we run out of pages, allocate all the remaining pages
-		PageCount = FreePageCount;
-	}
-
-	if (PageCount > FreePageCount || PageCount == 0)
+	if (FreePageCount == 0)
 	{
 		return false;
 	}
-	OutNumPagesAllocated = PageCount;
 
 	// Allocate spans from the free list head
 	uint32 NumPagesToFind = PageCount;
 	uint32 FoundPages = 0;
 	FPageSpan& HeadSpan = PageSpans[FreeSpanListHeadIndex];
-	uint32 StartSpanIndex = HeadSpan.NextSpanIndex;
-	uint32 SpanIndex = StartSpanIndex;
-	while (SpanIndex != FreeSpanListTailIndex && SpanIndex != InvalidIndex)
+	uint32 FirstSpanIndex = InvalidIndex;
+	uint32 LastSpanIndex = InvalidIndex;
+	uint32 SpanIndex = HeadSpan.NextSpanIndex;
+	while (SpanIndex != FreeSpanListTailIndex && SpanIndex != InvalidIndex && NumPagesToFind > 0)
 	{
 		FPageSpan& Span = PageSpans[SpanIndex];
-		if (NumPagesToFind <= Span.Count)
+		uint32 NextSpanIndex = Span.NextSpanIndex;
+
+		if (!FRHITransientAllocationFences::Contains(Span.Fences, Fences))
 		{
-			// Span is too big, so split it
-			if (Span.Count > NumPagesToFind)
+			if (NumPagesToFind <= Span.Count)
 			{
-				SplitSpan(SpanIndex, NumPagesToFind);
+				// Span is too big, so split it
+				if (Span.Count > NumPagesToFind)
+				{
+					SplitSpan(SpanIndex, NumPagesToFind);
+				}
+				check(NumPagesToFind == Span.Count);
 			}
-			check(NumPagesToFind == Span.Count);
+			Span.bAllocated = true;
 
-			// Move the head to point to the next free span
-			if (HeadSpan.NextSpanIndex != InvalidIndex)
+			NextSpanIndex = Span.NextSpanIndex;
+			Unlink(SpanIndex);
+			if (FirstSpanIndex == InvalidIndex)
+			{
+				FirstSpanIndex = LastSpanIndex = SpanIndex;
+			}
+			else
 			{
-				PageSpans[HeadSpan.NextSpanIndex].PrevSpanIndex = InvalidIndex;
+				InsertAfter(LastSpanIndex, SpanIndex);
+				LastSpanIndex = SpanIndex;
 			}
-			HeadSpan.NextSpanIndex = Span.NextSpanIndex;
-			if (Span.NextSpanIndex != InvalidIndex)
+
+			// Record the aliasing overlap between the resource we are allocating and the one that was deallocated.
+			if (Span.Resource)
 			{
-				PageSpans[Span.NextSpanIndex].PrevSpanIndex = FreeSpanListHeadIndex;
+				Resource->AddAliasingOverlap(Span.Resource, FRHITransientAllocationFences::GetAcquireFence(Span.Fences, Fences));
 			}
-			Span.NextSpanIndex = InvalidIndex;
+
+			check(NumPagesToFind >= Span.Count)
+			NumPagesToFind -= Span.Count;
 		}
-		Span.bAllocated = true;
-		NumPagesToFind -= Span.Count;
-		SpanIndex = Span.NextSpanIndex;
+
+		SpanIndex = NextSpanIndex;
 	}
-	check(NumPagesToFind == 0);
-	FreePageCount -= PageCount;
-#if UE_BUILD_DEBUG
+
+	const uint32 NumPagesAllocated = PageCount - NumPagesToFind;
+	if (NumPagesAllocated > 0)
+	{
+		FreePageCount -= NumPagesAllocated;
+		AllocationCount++;
+		OutSpanIndex = FirstSpanIndex;
+		OutNumPagesAllocated = NumPagesAllocated;
+	}
+
 	Validate();
-#endif
-	AllocationCount++;
-	OutSpanIndex = StartSpanIndex;
-	return true;
+	return NumPagesAllocated != 0;
 }
 
-void FRHITransientPageSpanAllocator::Deallocate(uint32 SpanIndex)
+void FRHITransientPageSpanAllocator::Deallocate(FRHITransientResource* Resource, const FRHITransientAllocationFences& Fences, uint32 SpanIndex)
 {
 	if (SpanIndex == InvalidIndex)
 	{
 		return;
 	}
 	check(AllocationCount > 0);
-	// Find the right span with which to merge this
 	while (SpanIndex != InvalidIndex)
 	{
 		FPageSpan& FreedSpan = PageSpans[SpanIndex];
 		check(FreedSpan.bAllocated);
 		FreePageCount += FreedSpan.Count;
 		uint32 NextSpanIndex = FreedSpan.NextSpanIndex;
+		FreedSpan.Resource = Resource;
+		FreedSpan.Fences = Fences;
 		FreedSpan.bAllocated = false;
-		if (!MergeFreeSpanIfPossible(SpanIndex))
-		{
-			// If we can't merge this span, just unlink and add it to the head (or tail)
-			Unlink(SpanIndex);
-
-			InsertAfter(FreeSpanListHeadIndex, SpanIndex);
-		}
+		Unlink(SpanIndex);
+		InsertAfter(FreeSpanListHeadIndex, SpanIndex);
 		SpanIndex = NextSpanIndex;
 	}
 	AllocationCount--;
 
-#if UE_BUILD_DEBUG
 	Validate();
-#endif
 }
 
 void FRHITransientPageSpanAllocator::SplitSpan(uint32 InSpanIndex, uint32 InPageCount)
@@ -1002,6 +976,8 @@ void FRHITransientPageSpanAllocator::SplitSpan(uint32 InSpanIndex, uint32 InPage
 	{
 		uint32 NewSpanIndex = AllocSpan();
 		FPageSpan& NewSpan = PageSpans[NewSpanIndex];
+		NewSpan.Resource = Span.Resource;
+		NewSpan.Fences = Span.Fences;
 		NewSpan.NextSpanIndex = Span.NextSpanIndex;
 		NewSpan.PrevSpanIndex = InSpanIndex;
 		NewSpan.Count = Span.Count - InPageCount;
@@ -1021,7 +997,7 @@ void FRHITransientPageSpanAllocator::SplitSpan(uint32 InSpanIndex, uint32 InPage
 	}
 }
 
-void FRHITransientPageSpanAllocator::MergeSpans(uint32 SpanIndex0, uint32 SpanIndex1, const bool bKeepSpan1)
+void FRHITransientPageSpanAllocator::MergeSpans(uint32 SpanIndex0, uint32 SpanIndex1)
 {
 	FPageSpan& Span0 = PageSpans[SpanIndex0];
 	FPageSpan& Span1 = PageSpans[SpanIndex1];
@@ -1030,28 +1006,64 @@ void FRHITransientPageSpanAllocator::MergeSpans(uint32 SpanIndex0, uint32 SpanIn
 	check(Span0.NextSpanIndex == SpanIndex1);
 	check(Span1.PrevSpanIndex == SpanIndex0);
 
-	uint32 SpanIndexToKeep = bKeepSpan1 ? SpanIndex1 : SpanIndex0;
-	uint32 SpanIndexToRemove = bKeepSpan1 ? SpanIndex0 : SpanIndex1;
+	uint32 SpanIndexToKeep = SpanIndex0;
+	uint32 SpanIndexToRemove = SpanIndex1;
 
 	// Update the PageToSpan mappings
 	PageToSpanStart[Span0.Offset] = SpanIndexToKeep;
 	PageToSpanStart[Span1.Offset] = InvalidIndex;
 	PageToSpanEnd[Span0.Offset + Span0.Count] = InvalidIndex; // Should match Span1.Offset
 	PageToSpanEnd[Span1.Offset + Span1.Count] = SpanIndexToKeep;
-	if (bKeepSpan1)
-	{
-		Span1.Offset = Span0.Offset;
-		Span1.Count += Span0.Count;
-	}
-	else
-	{
-		Span0.Count += Span1.Count;
-	}
+	Span0.Count += Span1.Count;
 
 	Unlink(SpanIndexToRemove);
 	ReleaseSpan(SpanIndexToRemove);
 }
 
+void FRHITransientPageSpanAllocator::Flush()
+{
+	uint32 PageIndex = 0;
+	while (PageIndex < MaxPageCount)
+	{
+		int32 SpanIndex = PageToSpanStart[PageIndex];
+		check(SpanIndex != InvalidIndex);
+		FPageSpan& Span = PageSpans[SpanIndex];
+
+		if (!Span.bAllocated)
+		{
+			Span.Resource = nullptr;
+			Span.Fences = {};
+
+			while (true)
+			{
+				// Can we merge this span with an existing free one to the right?
+				int32 NextSpanIndex = PageToSpanStart[Span.Offset + Span.Count];
+
+				if (NextSpanIndex == InvalidIndex)
+				{
+					break;
+				}
+
+				FPageSpan& NextPageSpan = PageSpans[NextSpanIndex];
+
+				if (NextPageSpan.bAllocated)
+				{
+					PageIndex += NextPageSpan.Count;
+					break;
+				}
+
+				NextPageSpan.Resource = nullptr;
+				Unlink(SpanIndex);
+				InsertBefore(NextSpanIndex, SpanIndex);
+				MergeSpans(SpanIndex, NextSpanIndex);
+				Validate();
+			}
+		}
+
+		PageIndex += Span.Count;
+	}
+}
+
 // Inserts a span after an existing span. The span to insert must be unlinked
 void FRHITransientPageSpanAllocator::InsertAfter(uint32 InsertPosition, uint32 InsertSpanIndex)
 {
@@ -1125,35 +1137,6 @@ uint32 FRHITransientPageSpanAllocator::GetAllocationPageCount(uint32 SpanIndex)
 	return Count;
 }
 
-bool FRHITransientPageSpanAllocator::MergeFreeSpanIfPossible(uint32 SpanIndex)
-{
-	FPageSpan& Span = PageSpans[SpanIndex];
-	check(!Span.bAllocated);
-	bool bMerged = false;
-
-	// Can we merge this span with an existing one to the left?
-	uint32 AdjSpanIndexPrev = PageToSpanEnd[Span.Offset];
-	if (AdjSpanIndexPrev != InvalidIndex && !PageSpans[AdjSpanIndexPrev].bAllocated)
-	{
-		Unlink(SpanIndex);
-		InsertAfter(AdjSpanIndexPrev, SpanIndex);
-		MergeSpans(AdjSpanIndexPrev, SpanIndex, true);
-		bMerged = true;
-	}
-
-	// Can we merge this span with an existing free one to the right?
-	uint32 AdjSpanIndexNext = PageToSpanStart[Span.Offset + Span.Count];
-	if (AdjSpanIndexNext != InvalidIndex && !PageSpans[AdjSpanIndexNext].bAllocated)
-	{
-		Unlink(SpanIndex);
-		InsertBefore(AdjSpanIndexNext, SpanIndex);
-		MergeSpans(SpanIndex, AdjSpanIndexNext, false);
-		bMerged = true;
-	}
-
-	return bMerged;
-}
-
 void FRHITransientPageSpanAllocator::Validate()
 {
 #if UE_BUILD_DEBUG
@@ -1171,6 +1154,7 @@ void FRHITransientPageSpanAllocator::Validate()
 	{
 		FPageSpan& Span = PageSpans[Index];
 		check(Span.PrevSpanIndex == PrevIndex);
+		check(Index == FreeSpanListHeadIndex || Index == FreeSpanListTailIndex || Span.Count != 0);
 		PrevIndex = Index;
 		FreeCount += Span.Count;
 	}
@@ -1187,7 +1171,7 @@ void FRHITransientPagePool::Allocate(FAllocationContext& Context)
 	uint32 PagesAllocated = 0;
 
 	uint32 PagesRemaining = Context.MaxAllocationPage > 0 ? FMath::Min(Context.PagesRemaining, Context.MaxAllocationPage) : Context.PagesRemaining;
-	if (Allocator.Allocate(PagesRemaining, PagesAllocated, SpanIndex))
+	if (Allocator.Allocate(&Context.Resource, Context.Fences, PagesRemaining, PagesAllocated, SpanIndex))
 	{
 		const uint64 DestinationGpuVirtualAddress = Context.GpuVirtualAddress + Context.PagesAllocated * Initializer.PageSize;
 		const uint32 PageSpanOffsetMin = PageSpans.Num();
@@ -1200,12 +1184,6 @@ void FRHITransientPagePool::Allocate(FAllocationContext& Context)
 		const  int32 AllocationIndex   = Context.Allocations.Num();
 		const uint64 AllocationHash    = CityHash64WithSeed((const char*)&PageSpans[PageSpanOffsetMin], PageSpanCount * sizeof(FRHITransientPageSpan), DestinationGpuVirtualAddress);
 
-			for (uint32 Index = PageSpanOffsetMin; Index < PageSpanOffsetMax; ++Index)
-			{
-				const FRHITransientPageSpan Span = PageSpans[Index];
-				OverlapTracker.Track(&Context.Resource, Span.Offset, Span.Offset + Span.Count);
-			}
-
 		FRHITransientPagePoolAllocation Allocation;
 		Allocation.Pool = this;
 		Allocation.Hash = AllocationHash;
@@ -1255,7 +1233,7 @@ void FRHITransientPagePool::Flush(FRHICommandListImmediate& RHICmdList)
 		PageSpans.Reserve(PageSpanCountMax);
 	}
 
-		OverlapTracker.Reset();
+	Allocator.Flush();
 }
 
 ///////////////////////////////////////////////////////////////////////////////////////////////////
@@ -1366,9 +1344,8 @@ void FRHITransientPagePoolCache::GarbageCollect()
 
 				PoolsToDelete.Emplace(PagePool);
 
-			#if 1 // Only delete one per frame. Deletion can be quite expensive.
+				// Only delete one per frame. Deletion can be quite expensive.
 				break;
-			#endif
 			}
 		}
 
@@ -1390,7 +1367,7 @@ void FRHITransientPagePoolCache::GarbageCollect()
 FRHITransientTexture* FRHITransientResourcePageAllocator::CreateTexture(
 	const FRHITextureCreateInfo& CreateInfo,
 	const TCHAR* DebugName,
-	uint32 PassIndex)
+	const FRHITransientAllocationFences& Fences)
 {
 	FRHITransientTexture* Texture = Textures.Acquire(ComputeHash(CreateInfo), [&](uint64 Hash)
 	{
@@ -1402,8 +1379,8 @@ FRHITransientTexture* FRHITransientResourcePageAllocator::CreateTexture(
 	const float FastPoolPercentageRequested = bFastPool ? CreateInfo.FastVRAMPercentage / 255.f : 0.f;
 
 	check(Texture);
-	Texture->Acquire(FRHICommandListImmediate::Get(), DebugName, PassIndex, CurrentCycle);
-	AllocateMemoryInternal(Texture, DebugName, PassIndex, bFastPool, FastPoolPercentageRequested);
+	Texture->Acquire(DebugName, Fences.GetSinglePipeline(), CurrentCycle);
+	AllocateMemoryInternal(Texture, DebugName, Fences, bFastPool, FastPoolPercentageRequested);
 	Stats.AllocateTexture(Texture->GetSize());
 	IF_RHICORE_TRANSIENT_ALLOCATOR_DEBUG(ActiveResources.Emplace(Texture));
 	return Texture;
@@ -1412,7 +1389,7 @@ FRHITransientTexture* FRHITransientResourcePageAllocator::CreateTexture(
 FRHITransientBuffer* FRHITransientResourcePageAllocator::CreateBuffer(
 	const FRHIBufferCreateInfo& CreateInfo,
 	const TCHAR* DebugName,
-	uint32 PassIndex)
+	const FRHITransientAllocationFences& Fences)
 {
 	FRHITransientBuffer* Buffer = Buffers.Acquire(ComputeHash(CreateInfo), [&](uint64 Hash)
 	{
@@ -1421,16 +1398,16 @@ FRHITransientBuffer* FRHITransientResourcePageAllocator::CreateBuffer(
 	});
 
 	check(Buffer);
-	Buffer->Acquire(FRHICommandListImmediate::Get(), DebugName, PassIndex, CurrentCycle);
-	AllocateMemoryInternal(Buffer, DebugName, PassIndex, EnumHasAnyFlags(CreateInfo.Usage, EBufferUsageFlags::FastVRAM), false);
+	Buffer->Acquire(DebugName, Fences.GetSinglePipeline(), CurrentCycle);
+	AllocateMemoryInternal(Buffer, DebugName, Fences, EnumHasAnyFlags(CreateInfo.Usage, EBufferUsageFlags::FastVRAM), false);
 	Stats.AllocateBuffer(Buffer->GetSize());
 	IF_RHICORE_TRANSIENT_ALLOCATOR_DEBUG(ActiveResources.Emplace(Buffer));
 	return Buffer;
 }
 
-void FRHITransientResourcePageAllocator::AllocateMemoryInternal(FRHITransientResource* Resource, const TCHAR* DebugName, uint32 PassIndex, bool bFastPoolRequested, float FastPoolPercentageRequested)
+void FRHITransientResourcePageAllocator::AllocateMemoryInternal(FRHITransientResource* Resource, const TCHAR* DebugName, const FRHITransientAllocationFences& Fences, bool bFastPoolRequested, float FastPoolPercentageRequested)
 {
-	FRHITransientPagePool::FAllocationContext AllocationContext(*Resource, PageSize);
+	FRHITransientPagePool::FAllocationContext AllocationContext(*Resource, Fences, PageSize);
 
 	if (bFastPoolRequested && FastPagePool)
 	{
@@ -1465,25 +1442,25 @@ void FRHITransientResourcePageAllocator::AllocateMemoryInternal(FRHITransientRes
 	PageSpanCount     += AllocationContext.PageSpansAllocated;
 }
 
-void FRHITransientResourcePageAllocator::DeallocateMemoryInternal(FRHITransientResource* Resource, uint32 PassIndex)
+void FRHITransientResourcePageAllocator::DeallocateMemoryInternal(FRHITransientResource* Resource, const FRHITransientAllocationFences& Fences)
 {
-	Resource->Discard(PassIndex);
+	Resource->Discard();
 
 	for (const FRHITransientPagePoolAllocation& Allocation : Resource->GetPageAllocation().PoolAllocations)
 	{
-		Allocation.Pool->Deallocate(Allocation.SpanIndex);
+		Allocation.Pool->Deallocate(Resource, Fences, Allocation.SpanIndex);
 	}
 }
 
-void FRHITransientResourcePageAllocator::DeallocateMemory(FRHITransientTexture* Texture, uint32 PassIndex)
+void FRHITransientResourcePageAllocator::DeallocateMemory(FRHITransientTexture* Texture, const FRHITransientAllocationFences& Fences)
 {
-	DeallocateMemoryInternal(Texture, PassIndex);
+	DeallocateMemoryInternal(Texture, Fences);
 	Stats.DeallocateTexture(Texture->GetSize());
 }
 
-void FRHITransientResourcePageAllocator::DeallocateMemory(FRHITransientBuffer* Buffer, uint32 PassIndex)
+void FRHITransientResourcePageAllocator::DeallocateMemory(FRHITransientBuffer* Buffer, const FRHITransientAllocationFences& Fences)
 {
-	DeallocateMemoryInternal(Buffer, PassIndex);
+	DeallocateMemoryInternal(Buffer, Fences);
 	Stats.DeallocateBuffer(Buffer->GetSize());
 }
 
@@ -1583,12 +1560,28 @@ void FRHITransientResourcePageAllocator::Flush(FRHICommandListImmediate& RHICmdL
 
 	if (Stats.HasDeallocations())
 	{
-		const int32 FirstForfeitIndex = Algo::Partition(PagePools.GetData(), PagePools.Num(), [](const FRHITransientPagePool* PagePool) { return !PagePool->IsEmpty(); });
-		PagePoolCache.Forfeit(MakeArrayView(PagePools.GetData() + FirstForfeitIndex, PagePools.Num() - FirstForfeitIndex));
-		PagePools.SetNum(FirstForfeitIndex, EAllowShrinking::No);
+		TArray<FRHITransientPagePool*, FConcurrentLinearArrayAllocator> EmptyPagePools;
+		TArray<FRHITransientPagePool*, FConcurrentLinearArrayAllocator> ActivePagePools;
+		EmptyPagePools.Reserve(PagePools.Num());
+		ActivePagePools.Reserve(PagePools.Num());
+
+		for (FRHITransientPagePool* PagePool : PagePools)
+		{
+			if (PagePool->IsEmpty())
+			{
+				EmptyPagePools.Emplace(PagePool);
+			}
+			else
+			{
+				ActivePagePools.Emplace(PagePool);
+			}
+		}
+
+		PagePoolCache.Forfeit(EmptyPagePools);
+		PagePools = ActivePagePools;
 	}
 
-	RHICmdList.EnqueueLambda([&PagePoolCache = PagePoolCache, Stats = Stats](FRHICommandListImmediate&)
+	RHICmdList.EnqueueLambda([&PagePoolCache = PagePoolCache, Stats = Stats](FRHICommandListBase&)
 	{
 		PagePoolCache.Stats.Accumulate(Stats);
 	});
@@ -1596,4 +1589,4 @@ void FRHITransientResourcePageAllocator::Flush(FRHICommandListImmediate& RHICmdL
 	Stats.Reset();
 
 	CurrentCycle++;
-}
+}
\ No newline at end of file
diff --git a/Engine/Source/Runtime/RHICore/Public/RHICoreTransientResourceAllocator.h b/Engine/Source/Runtime/RHICore/Public/RHICoreTransientResourceAllocator.h
--- a/Engine/Source/Runtime/RHICore/Public/RHICoreTransientResourceAllocator.h
+++ b/Engine/Source/Runtime/RHICore/Public/RHICoreTransientResourceAllocator.h
@@ -137,25 +137,6 @@ public:
 	uint64 AliasedSize = 0;
 };
 
-/** Tracks resource allocations on the heap and adds overlap events to transient resources. */
-class FRHITransientResourceOverlapTracker
-{
-public:
-	RHICORE_API void Track(FRHITransientResource* InResource, uint32 PageOffsetMin, uint32 PageOffsetMax);
-	RHICORE_API void Reset();
-
-private:
-	struct FResourceRange
-	{
-		uint64 PageOffsetMin : 24;
-		uint64 PageOffsetMax : 24;
-		uint64 ResourceIndex : 16;
-	};
-
-	TArray<FResourceRange> ResourceRanges;
-	TArray<FRHITransientResource*> Resources;
-};
-
 /** An RHI transient resource cache designed to optimize fetches for resources placed into a heap with an offset.
  *  The cache has a fixed capacity whereby no garbage collection will occur. Once that capacity is exceeded, garbage
  *  collection is invoked on resources older than a specified generation (where generation is incremented with each
@@ -327,11 +308,25 @@ class FRHITransientResourceHeapAllocator;
 class FRHITransientHeapAllocator
 {
 public:
+	struct FAliasingOverlap
+	{
+		FAliasingOverlap() = default;
+		FAliasingOverlap(FRHITransientResource* InResource, uint32 InAcquireFence)
+			: Resource(InResource)
+			, AcquireFence(InAcquireFence)
+		{}
+
+		FRHITransientResource* Resource;
+		uint32 AcquireFence;
+	};
+
 	RHICORE_API FRHITransientHeapAllocator(uint64 Capacity, uint32 Alignment);
 
-	RHICORE_API FRHITransientHeapAllocation Allocate(uint64 Size, uint32 Alignment);
+	RHICORE_API FRHITransientHeapAllocation Allocate(const FRHITransientAllocationFences& Fences, uint64 Size, uint32 Alignment, TArray<FAliasingOverlap>& OutAliasingOverlaps);
 
-	RHICORE_API void Deallocate(FRHITransientHeapAllocation Allocation);
+	RHICORE_API void Deallocate(FRHITransientResource* Resource, const FRHITransientAllocationFences& Fences);
+
+	RHICORE_API void Flush();
 
 	void SetGpuVirtualAddress(uint64 InGpuVirtualAddress)
 	{
@@ -354,9 +349,12 @@ private:
 
 	struct FRange
 	{
-		uint64 Size{};
-		uint64 Offset{};
+		FRHITransientResource* Resource = nullptr;
+		FRHITransientAllocationFences Fences;
+		uint64 Size = 0;
+		uint64 Offset = 0;
 		FRangeHandle NextFreeHandle = InvalidRangeHandle;
+		FRangeHandle PrevFreeHandle = InvalidRangeHandle;
 
 		FORCEINLINE uint64 GetStart() const { return Offset; }
 		FORCEINLINE uint64 GetEnd() const { return Size + Offset; }
@@ -377,11 +375,13 @@ private:
 		return FRangeHandle(Ranges.Num() - 1);
 	}
 
-	FRangeHandle InsertRange(FRangeHandle PreviousHandle, uint64 Offset, uint64 Size)
+	FRangeHandle InsertRange(FRangeHandle PreviousHandle, FRHITransientResource* Resource, const FRHITransientAllocationFences& Fences, uint64 Offset, uint64 Size)
 	{
 		FRangeHandle Handle = CreateRange();
 
 		FRange& CurrentRange = Ranges[Handle];
+		CurrentRange.Resource = Resource;
+		CurrentRange.Fences = Fences;
 		CurrentRange.Offset = Offset;
 		CurrentRange.Size = Size;
 
@@ -392,15 +392,20 @@ private:
 		return Handle;
 	}
 
-	void RemoveRange(FRangeHandle PreviousHandle, FRangeHandle CurrentHandle)
+	FRangeHandle RemoveRange(FRangeHandle PreviousHandle, FRangeHandle CurrentHandle)
 	{
 		FRange& PreviousRange = Ranges[PreviousHandle];
 		FRange& CurrentRange = Ranges[CurrentHandle];
 
+		FRangeHandle NextCurrentHandle = CurrentRange.NextFreeHandle;
+
+		check(PreviousRange.NextFreeHandle == CurrentHandle);
 		PreviousRange.NextFreeHandle = CurrentRange.NextFreeHandle;
 		CurrentRange.NextFreeHandle = InvalidRangeHandle;
+		CurrentRange.Resource = nullptr;
 
 		RangeFreeList.Add(CurrentHandle);
+		return NextCurrentHandle;
 	}
 
 	struct FFindResult
@@ -410,8 +415,6 @@ private:
 		FRangeHandle FoundHandle = InvalidRangeHandle;
 	};
 
-	RHICORE_API FFindResult FindFreeRange(uint64 Size, uint32 Alignment);
-
 	RHICORE_API void Validate();
 
 	uint64 GpuVirtualAddress = 0;
@@ -506,24 +509,24 @@ public:
 	RHICORE_API FRHITransientTexture* CreateTexture(
 		const FRHITextureCreateInfo& CreateInfo,
 		const TCHAR* DebugName,
-		uint32 PassIndex,
+		const FRHITransientAllocationFences& Fences,
 		uint64 CurrentAllocatorCycle,
 		uint64 TextureSize,
 		uint32 TextureAlignment,
 		FCreateTextureFunction CreateTextureFunction);
 
-	RHICORE_API void DeallocateMemory(FRHITransientTexture* Texture, uint32 PassIndex);
+	RHICORE_API void DeallocateMemory(FRHITransientTexture* Texture, const FRHITransientAllocationFences& Fences);
 
 	RHICORE_API FRHITransientBuffer* CreateBuffer(
 		const FRHIBufferCreateInfo& CreateInfo,
 		const TCHAR* DebugName,
-		uint32 PassIndex,
+		const FRHITransientAllocationFences& Fences,
 		uint64 CurrentAllocatorCycle,
 		uint64 BufferSize,
 		uint32 BufferAlignment,
 		FCreateBufferFunction CreateBufferFunction);
 
-	RHICORE_API void DeallocateMemory(FRHITransientBuffer* Buffer, uint32 PassIndex);
+	RHICORE_API void DeallocateMemory(FRHITransientBuffer* Buffer, const FRHITransientAllocationFences& Fences);
 
 	RHICORE_API void Flush(uint64 CurrentAllocatorCycle, FRHITransientMemoryStats& OutMemoryStats, FRHITransientAllocationStats* OutAllocationStats);
 
@@ -555,17 +558,18 @@ protected:
 	}
 
 private:
-	RHICORE_API void AllocateMemoryInternal(FRHITransientResource* Resource, const TCHAR* Name, uint32 PassIndex, uint64 CurrentAllocatorCycle, const FRHITransientHeapAllocation& Allocation);
-	RHICORE_API void DeallocateMemoryInternal(FRHITransientResource* Resource, uint32 PassIndex);
+	RHICORE_API void AllocateMemoryInternal(FRHITransientResource* Resource, const FRHITransientHeapAllocation& Allocation);
+	RHICORE_API void DeallocateMemoryInternal(FRHITransientResource* Resource, const FRHITransientAllocationFences& Fences);
 
 	FInitializer Initializer;
 	FRHITransientHeapAllocator Allocator;
 
 	uint64 LastUsedGarbageCollectCycle = 0;
 	uint64 CommitSize = 0;
+	uint64 CommitSizeMax = 0;
 	uint32 AlignmentLog2;
 
-	FRHITransientResourceOverlapTracker OverlapTracker;
+	TArray<FRHITransientHeapAllocator::FAliasingOverlap> AliasingOverlaps;
 
 	FRHITransientMemoryStats Stats;
 	TRHITransientResourceCache<FRHITransientTexture> Textures;
@@ -594,6 +598,9 @@ public:
 
 		static const uint32 kDefaultResourceCacheSize = 256;
 
+		// The minimum size to use when creating the first heap. This is the default but can grow based on allocations.
+		uint64 MinimumFirstHeapSize = 0;
+
 		// The minimum size to use when creating a heap. This is the default but can grow based on allocations.
 		uint64 MinimumHeapSize = 0;
 
@@ -611,6 +618,9 @@ public:
 
 		// Whether all heaps should be created with the AllowAll heap flag.
 		bool bSupportsAllHeapFlags = true;
+
+		// Whether all heaps support mapping physical pages to the commit size. If false the physical memory usage is represented by the capacity instead.
+		bool bSupportsVirtualMapping = false;
 	};
 
 	FRHITransientHeapCache(const FInitializer& InInitializer)
@@ -650,7 +660,6 @@ private:
 	TArray<FRHITransientHeap*> LiveList;
 	TArray<FRHITransientHeap*> FreeList;
 	uint64 GarbageCollectCycle = 0;
-	uint64 TotalMemoryCapacity = 0;
 
 	friend FRHITransientResourceHeapAllocator;
 };
@@ -663,11 +672,14 @@ public:
 		: HeapCache(InHeapCache)
 	{}
 
+	// Sets the create mode for allocations.
+	RHICORE_API void SetCreateMode(ERHITransientResourceCreateMode InCreateMode) override;
+
 	// Deallocates a texture from its parent heap. Provide the current platform fence value used to update the heap.
-	RHICORE_API void DeallocateMemory(FRHITransientTexture* Texture, uint32 PassIndex) override;
+	RHICORE_API void DeallocateMemory(FRHITransientTexture* Texture, const FRHITransientAllocationFences& Fences) override;
 
 	// Deallocates a buffer from its parent heap. Provide the current platform fence value used to update the heap.
-	RHICORE_API void DeallocateMemory(FRHITransientBuffer* Buffer, uint32 PassIndex) override;
+	RHICORE_API void DeallocateMemory(FRHITransientBuffer* Buffer, const FRHITransientAllocationFences& Fences) override;
 
 	// Called to flush any active allocations prior to rendering.
 	RHICORE_API void Flush(FRHICommandListImmediate& RHICmdList, FRHITransientAllocationStats* OutAllocationStats) override;
@@ -677,6 +689,22 @@ public:
 
 	FRHITransientHeapCache& HeapCache;
 
+	template <typename TransientResourceType, typename LambdaType, typename ResourceCreateInfo>
+	TransientResourceType* CreateTransientResource(LambdaType&& Lambda, uint64 Hash, uint64 Size, const ResourceCreateInfo& CreateInfo)
+	{
+		TransientResourceType* Resource;
+		if (CreateMode == ERHITransientResourceCreateMode::Inline)
+		{
+			typename TransientResourceType::FResourceTaskResult TaskResult = Lambda();
+			Resource = new TransientResourceType(TaskResult.Resource.GetReference(), TaskResult.GpuVirtualAddress, Hash, Size, ERHITransientAllocationType::Heap, CreateInfo);
+		}
+		else
+		{
+			Resource = new TransientResourceType(UE::Tasks::Launch(UE_SOURCE_LOCATION, MoveTemp(Lambda), LowLevelTasks::ETaskPriority::High), Hash, Size, ERHITransientAllocationType::Heap, CreateInfo);
+		}
+		return Resource;
+	}
+
 protected:
 	/** Allocates a texture on a heap at a specific offset, returning a cached RHI transient texture pointer, or null
 	 *  if the allocation failed. TextureSize and TextureAlignment are platform specific and must be derived from the
@@ -686,7 +714,7 @@ protected:
 	RHICORE_API FRHITransientTexture* CreateTextureInternal(
 		const FRHITextureCreateInfo& CreateInfo,
 		const TCHAR* DebugName,
-		uint32 PassIndex,
+		const FRHITransientAllocationFences& Fences,
 		uint64 TextureSize,
 		uint32 TextureAlignment,
 		FRHITransientHeap::FCreateTextureFunction CreateTextureFunction);
@@ -699,7 +727,7 @@ protected:
 	RHICORE_API FRHITransientBuffer* CreateBufferInternal(
 		const FRHIBufferCreateInfo& CreateInfo,
 		const TCHAR* DebugName,
-		uint32 PassIndex,
+		const FRHITransientAllocationFences& Fences,
 		uint32 BufferSize,
 		uint32 BufferAlignment,
 		FRHITransientHeap::FCreateBufferFunction CreateBufferFunction);
@@ -708,6 +736,7 @@ private:
 	TArray<FRHITransientHeap*> Heaps;
 	uint64 CurrentCycle = 0;
 	uint32 DeallocationCount = 0;
+	ERHITransientResourceCreateMode CreateMode = ERHITransientResourceCreateMode::Inline;
 
 	IF_RHICORE_TRANSIENT_ALLOCATOR_DEBUG(TSet<FRHITransientResource*> ActiveResources);
 };
@@ -733,9 +762,11 @@ public:
 
 	RHICORE_API void Reset();
 
-	RHICORE_API bool Allocate(uint32 PageCount, uint32& NumPagesAllocated, uint32& SpanIndex);
+	RHICORE_API bool Allocate(FRHITransientResource* Resource, const FRHITransientAllocationFences& Fences, uint32 PageCount, uint32& NumPagesAllocated, uint32& OutSpanIndex);
 
-	RHICORE_API void Deallocate(uint32 SpanIndex);
+	RHICORE_API void Deallocate(FRHITransientResource* Resource, const FRHITransientAllocationFences& Fences, uint32 SpanIndex);
+
+	RHICORE_API void Flush();
 
 	template <typename SpanArrayType>
 	void GetSpanArray(uint32 SpanIndex, SpanArrayType& OutPageSpans) const
@@ -777,6 +808,8 @@ private:
 	{
 		const bool IsLinked() { return (NextSpanIndex != InvalidIndex || PrevSpanIndex != InvalidIndex); }
 
+		FRHITransientResource* Resource = nullptr;
+		FRHITransientAllocationFences Fences;
 		uint32 NextSpanIndex = InvalidIndex;
 		uint32 PrevSpanIndex = 0;
 		bool bAllocated = false;
@@ -788,7 +821,7 @@ private:
 	RHICORE_API void SplitSpan(uint32 SpanIndex, uint32 PageCount);
 
 	// Merges two spans. They must be adjacent and in the same list
-	RHICORE_API void MergeSpans(uint32 SpanIndex0, uint32 SpanIndex1, const bool bKeepSpan1);
+	RHICORE_API void MergeSpans(uint32 SpanIndex0, uint32 SpanIndex1);
 
 	// Inserts a span after an existing span. The span to insert must be unlinked
 	RHICORE_API void InsertAfter(uint32 InsertPosition, uint32 InsertSpanIndex);
@@ -817,9 +850,6 @@ private:
 		check(UnusedSpanListCount <= MaxPageCount);
 	}
 
-	// Merges a span with existing neighbours in the free list if they exist
-	RHICORE_API bool MergeFreeSpanIfPossible(uint32 SpanIndex);
-
 	RHICORE_API void Validate();
 
 	uint32 GetFirstSpanIndex() const
@@ -861,13 +891,14 @@ public:
 
 	struct FAllocationContext
 	{
-		FAllocationContext(FRHITransientResource& InResource, uint32 InPageSize)
+		FAllocationContext(FRHITransientResource& InResource, const FRHITransientAllocationFences& InFences, uint32 InPageSize)
 			: Resource(InResource)
 			, Allocations(InResource.GetPageAllocation().PoolAllocations)
 			, Spans(InResource.GetPageAllocation().Spans)
 			, AllocationsBefore(Allocations)
 			, GpuVirtualAddress(Resource.GetGpuVirtualAddress())
 			, Size(Align(Resource.GetSize(), InPageSize))
+			, Fences(InFences)
 			, PagesRemaining(Size / InPageSize)
 		{
 			Allocations.Reset();
@@ -882,6 +913,7 @@ public:
 		const TArray<FRHITransientPagePoolAllocation, TInlineAllocator<8>> AllocationsBefore;
 		const uint64 GpuVirtualAddress;
 		const uint64 Size;
+		const FRHITransientAllocationFences Fences;
 
 		uint32 AllocationCount = 0;
 		uint32 AllocationMatchingCount = 0;
@@ -894,9 +926,9 @@ public:
 
 	RHICORE_API void Allocate(FAllocationContext& AllocationContext);
 
-	void Deallocate(uint32 SpanIndex)
+	void Deallocate(FRHITransientResource* Resource, const FRHITransientAllocationFences& Fences, uint32 SpanIndex)
 	{
-		Allocator.Deallocate(SpanIndex);
+		Allocator.Deallocate(Resource, Fences, SpanIndex);
 	}
 
 	RHICORE_API void Flush(FRHICommandListImmediate& RHICmdList);
@@ -953,7 +985,6 @@ private:
 	//////////////////////////////////////////////////////////////////////////
 
 	FRHITransientPageSpanAllocator Allocator;
-	FRHITransientResourceOverlapTracker OverlapTracker;
 
 	TArray<FPageMapRequest> PageMapRequests;
 	TArray<FRHITransientPageSpan> PageSpans;
@@ -1056,10 +1087,10 @@ public:
 		FastPagePool = PagePoolCache.GetFastPagePool();
 	}
 
-	RHICORE_API FRHITransientTexture* CreateTexture(const FRHITextureCreateInfo& CreateInfo, const TCHAR* DebugName, uint32 PassIndex) override;
-	RHICORE_API FRHITransientBuffer* CreateBuffer(const FRHIBufferCreateInfo& CreateInfo, const TCHAR* DebugName, uint32 PassIndex) override;
-	RHICORE_API void DeallocateMemory(FRHITransientTexture* Texture, uint32 PassIndex) override;
-	RHICORE_API void DeallocateMemory(FRHITransientBuffer* Buffer, uint32 PassIndex) override;
+	RHICORE_API FRHITransientTexture* CreateTexture(const FRHITextureCreateInfo& CreateInfo, const TCHAR* DebugName, const FRHITransientAllocationFences& Fences) override;
+	RHICORE_API FRHITransientBuffer* CreateBuffer(const FRHIBufferCreateInfo& CreateInfo, const TCHAR* DebugName, const FRHITransientAllocationFences& Fences) override;
+	RHICORE_API void DeallocateMemory(FRHITransientTexture* Texture, const FRHITransientAllocationFences& Fences) override;
+	RHICORE_API void DeallocateMemory(FRHITransientBuffer* Buffer, const FRHITransientAllocationFences& Fences) override;
 	RHICORE_API void Flush(FRHICommandListImmediate& RHICmdList, FRHITransientAllocationStats* OutAllocationStats) override;
 
 	uint32 GetPageSize() const { return PageSize; }
@@ -1073,8 +1104,8 @@ private:
 	static constexpr uint64 KB = 1024;
 	static constexpr uint64 MB = 1024 * KB;
 
-	RHICORE_API void AllocateMemoryInternal(FRHITransientResource* Resource, const TCHAR* DebugName, uint32 PassIndex, bool bFastPoolRequested, float FastPoolPercentageRequested);
-	RHICORE_API void DeallocateMemoryInternal(FRHITransientResource* Resource, uint32 PassIndex);
+	RHICORE_API void AllocateMemoryInternal(FRHITransientResource* Resource, const TCHAR* DebugName, const FRHITransientAllocationFences& Fences, bool bFastPoolRequested, float FastPoolPercentageRequested);
+	RHICORE_API void DeallocateMemoryInternal(FRHITransientResource* Resource, const FRHITransientAllocationFences& Fences);
 
 	//////////////////////////////////////////////////////////////////////////
 	//! Platform API
diff --git a/Engine/Source/Runtime/RenderCore/Private/RenderGraphBuilder.cpp b/Engine/Source/Runtime/RenderCore/Private/RenderGraphBuilder.cpp
--- a/Engine/Source/Runtime/RenderCore/Private/RenderGraphBuilder.cpp
+++ b/Engine/Source/Runtime/RenderCore/Private/RenderGraphBuilder.cpp
@@ -652,7 +652,7 @@ const TRefCountPtr<FRDGPooledBuffer>& FRDGBuilder::ConvertToExternalBuffer(FRDGB
 		Buffer->bExternal = 1;
 		if (!Buffer->ResourceRHI)
 		{
-			SetPooledBufferRHI(Buffer, AllocatePooledBufferRHI(RHICmdList, Buffer));
+			SetExternalPooledBufferRHI(Buffer, AllocatePooledBufferRHI(RHICmdList, Buffer));
 		}
 		ExternalBuffers.FindOrAdd(Buffer->GetRHIUnchecked(), Buffer);
 		AsyncSetupQueue.Push(FAsyncSetupOp::CullRootBuffer(Buffer));
@@ -668,7 +668,7 @@ const TRefCountPtr<IPooledRenderTarget>& FRDGBuilder::ConvertToExternalTexture(F
 		Texture->bExternal = 1;
 		if (!Texture->ResourceRHI)
 		{
-			SetPooledRenderTargetRHI(Texture, AllocatePooledRenderTargetRHI(RHICmdList, Texture));
+			SetExternalPooledRenderTargetRHI(Texture, AllocatePooledRenderTargetRHI(RHICmdList, Texture));
 		}
 		ExternalTextures.FindOrAdd(Texture->GetRHIUnchecked(), Texture);
 		AsyncSetupQueue.Push(FAsyncSetupOp::CullRootTexture(Texture));
@@ -1029,7 +1029,7 @@ FRDGTexture* FRDGBuilder::RegisterExternalTexture(
 
 	const FRDGTextureDesc Desc = Translate(ExternalPooledTexture->GetDesc());
 	FRDGTexture* Texture = Textures.Allocate(Allocators.Root, Name, Desc, Flags);
-	SetPooledRenderTargetRHI(Texture, ExternalPooledTexture.GetReference());
+	SetExternalPooledRenderTargetRHI(Texture, ExternalPooledTexture.GetReference());
 	Texture->bExternal = true;
 	ExternalTextures.FindOrAdd(Texture->GetRHIUnchecked(), Texture);
 
@@ -1065,7 +1065,7 @@ FRDGBufferRef FRDGBuilder::RegisterExternalBuffer(
 	}
 
 	FRDGBuffer* Buffer = Buffers.Allocate(Allocators.Root, Name, ExternalPooledBuffer->Desc, Flags);
-	SetPooledBufferRHI(Buffer, ExternalPooledBuffer);
+	SetExternalPooledBufferRHI(Buffer, ExternalPooledBuffer);
 	Buffer->bExternal = true;
 	ExternalBuffers.FindOrAdd(Buffer->GetRHIUnchecked(), Buffer);
 
@@ -1144,9 +1144,15 @@ bool FRDGBuilder::AddCullingDependency(FRDGProducerStatesByPipeline& LastProduce
 		}
 	}
 
+	FRDGProducerState& LastProducer = LastProducers[NextPipeline];
+
 	if (IsWritableAccess(NextState.Access))
 	{
-		FRDGProducerState& LastProducer = LastProducers[NextPipeline];
+		// Add a dependency between the last read of a resource on the other pipe and the new write (this is necessary for async compute fencing).
+		if (FRDGPass* PassIfReadAccess = LastProducers[NextPipeline == ERHIPipeline::Graphics ? ERHIPipeline::AsyncCompute : ERHIPipeline::Graphics].PassIfReadAccess)
+		{
+			AddPassDependency(PassIfReadAccess, NextState.Pass);
+		}
 
 		// A separate producer pass is tracked for UAV -> UAV dependencies that are skipped. Consider the following scenario:
 		//
@@ -1177,8 +1183,14 @@ bool FRDGBuilder::AddCullingDependency(FRDGProducerStatesByPipeline& LastProduce
 		LastProducer.Access             = NextState.Access;
 		LastProducer.Pass               = NextState.Pass;
 		LastProducer.NoUAVBarrierHandle = NextState.NoUAVBarrierHandle;
+		LastProducer.PassIfReadAccess = nullptr;
 		return true;
 	}
+	else
+	{
+		LastProducer.PassIfReadAccess = NextState.Pass;
+	}
+
 	return false;
 }
 
@@ -1470,6 +1482,8 @@ void FRDGBuilder::Compile()
 	{
 		SCOPED_NAMED_EVENT(AsyncComputeFences, FColor::Emerald);
 
+		const bool bAsyncComputeTransientAliasing = IsAsyncComputeTransientAliasingEnabled();
+
 		// Establishes fork / join overlap regions for async compute. This is used for fencing as well as resource
 		// allocation / deallocation. Async compute passes can't allocate / release their resource references until
 		// the fork / join is complete, since the two pipes run in parallel. Therefore, all resource lifetimes on
@@ -1491,7 +1505,12 @@ void FRDGBuilder::Compile()
 			FRDGPass* GraphicsForkPass = Passes[GraphicsForkPassHandle];
 
 			AsyncComputePass->GraphicsForkPass = GraphicsForkPassHandle;
-			Passes[GraphicsForkPass->PrologueBarrierPass]->ResourcesToBegin.Add(AsyncComputePass);
+
+			if (!bAsyncComputeTransientAliasing)
+			{
+				AsyncComputePass->ResourcesToBegin.Reset();
+				Passes[GraphicsForkPass->PrologueBarrierPass]->ResourcesToBegin.Add(AsyncComputePass);
+			}
 
 			if (CurrentGraphicsForkPassHandle != GraphicsForkPassHandle)
 			{
@@ -1544,7 +1563,12 @@ void FRDGBuilder::Compile()
 			FRDGPass* GraphicsJoinPass = Passes[GraphicsJoinPassHandle];
 
 			AsyncComputePass->GraphicsJoinPass = GraphicsJoinPassHandle;
-			Passes[GraphicsJoinPass->EpilogueBarrierPass]->ResourcesToEnd.Add(AsyncComputePass);
+
+			if (!bAsyncComputeTransientAliasing)
+			{
+				AsyncComputePass->ResourcesToEnd.Reset();
+				Passes[GraphicsJoinPass->EpilogueBarrierPass]->ResourcesToEnd.Add(AsyncComputePass);
+			}
 
 			if (CurrentGraphicsJoinPassHandle != GraphicsJoinPassHandle)
 			{
@@ -1747,6 +1771,8 @@ void FRDGBuilder::Execute()
 
 			Buffers.Enumerate([&] (FRDGBuffer* Buffer)
 			{
+				Buffer->LastPasses = {};
+
 				if (Buffer->ResourceRHI || Buffer->bQueuedForUpload)
 				{
 					Buffer->bCollectForAllocate = false;
@@ -1760,6 +1786,8 @@ void FRDGBuilder::Execute()
 
 			Textures.Enumerate([&] (FRDGTexture* Texture)
 			{
+				Texture->LastPasses = {};
+
 				if (Texture->ResourceRHI)
 				{
 					Texture->bCollectForAllocate = false;
@@ -1802,6 +1830,8 @@ void FRDGBuilder::Execute()
 			SCOPE_CYCLE_COUNTER(STAT_RDG_CollectResourcesTime);
 			CSV_SCOPED_TIMING_STAT_EXCLUSIVE(RDG_CollectResources);
 			SCOPED_NAMED_EVENT_TEXT("FRDGBuilder::CollectResources", FColor::Magenta);
+			
+			PrepareCollectResourcesTask.Wait();
 
 			EnumerateExtendedLifetimeResources(Textures, [](FRDGTexture* Texture)
 			{
@@ -1813,8 +1843,6 @@ void FRDGBuilder::Execute()
 				++Buffer->ReferenceCount;
 			});
 
-			PrepareCollectResourcesTask.Wait();
-
 			// Null out any culled external resources so that the reference is freed up.
 
 			for (const auto& Pair : ExternalTextures)
@@ -1823,7 +1851,7 @@ void FRDGBuilder::Execute()
 
 				if (Texture->IsCulled())
 				{
-					CollectDeallocateTexture(CollectResourceContext, ProloguePassHandle, Texture, 0);
+					CollectDeallocateTexture(CollectResourceContext, ERHIPipeline::Graphics, ProloguePassHandle, Texture, 0);
 				}
 			}
 
@@ -1833,7 +1861,7 @@ void FRDGBuilder::Execute()
 
 				if (Buffer->IsCulled())
 				{
-					CollectDeallocateBuffer(CollectResourceContext, ProloguePassHandle, Buffer, 0);
+					CollectDeallocateBuffer(CollectResourceContext, ERHIPipeline::Graphics, ProloguePassHandle, Buffer, 0);
 				}
 			}
 
@@ -1850,12 +1878,12 @@ void FRDGBuilder::Execute()
 
 			EnumerateExtendedLifetimeResources(Textures, [&](FRDGTextureRef Texture)
 			{
-				CollectDeallocateTexture(CollectResourceContext, EpiloguePassHandle, Texture, 1);
+				CollectDeallocateTexture(CollectResourceContext, ERHIPipeline::Graphics, EpiloguePassHandle, Texture, 1);
 			});
 
 			EnumerateExtendedLifetimeResources(Buffers, [&](FRDGBufferRef Buffer)
 			{
-				CollectDeallocateBuffer(CollectResourceContext, EpiloguePassHandle, Buffer, 1);
+				CollectDeallocateBuffer(CollectResourceContext, ERHIPipeline::Graphics, EpiloguePassHandle, Buffer, 1);
 			});
 
 			BufferNumElementsCallbacksTask.Wait();
@@ -1930,7 +1958,7 @@ void FRDGBuilder::Execute()
 		SCOPED_NAMED_EVENT_TEXT("FRDGBuilder::ExecutePasses", FColor::Magenta);
 		SCOPE_CYCLE_COUNTER(STAT_RDG_ExecuteTime);
 		CSV_SCOPED_TIMING_STAT_EXCLUSIVE(RDG_Execute);
-		
+
 		FRDGPass* PrevSerialPass = nullptr;
 		TArray<FRHICommandListImmediate::FQueuedCommandList> QueuedCmdLists;
 
@@ -2245,9 +2273,9 @@ void FRDGBuilder::SetupPassResources(FRDGPass* Pass)
 
 		FRDGPass::FTextureState* PassState;
 
-		if (Texture->LastPass != PassHandle)
+		if (Texture->LastPasses[PassPipeline] != PassHandle)
 		{
-			Texture->LastPass = PassHandle;
+			Texture->LastPasses[PassPipeline] = PassHandle;
 			Texture->PassStateIndex = Pass->TextureStates.Num();
 
 			PassState = &Pass->TextureStates.Emplace_GetRef(Texture);
@@ -2304,9 +2332,9 @@ void FRDGBuilder::SetupPassResources(FRDGPass* Pass)
 
 		FRDGPass::FBufferState* PassState;
 
-		if (Buffer->LastPass != PassHandle)
+		if (Buffer->LastPasses[PassPipeline] != PassHandle)
 		{
-			Buffer->LastPass = PassHandle;
+			Buffer->LastPasses[PassPipeline] = PassHandle;
 			Buffer->PassStateIndex = Pass->BufferStates.Num();
 
 			PassState = &Pass->BufferStates.Emplace_GetRef(Buffer);
@@ -2364,12 +2392,8 @@ void FRDGBuilder::SetupPassInternals(FRDGPass* Pass)
 
 	Pass->PrologueBarrierPass = PassHandle;
 	Pass->EpilogueBarrierPass = PassHandle;
-
-	if (Pass->Pipeline == ERHIPipeline::Graphics)
-	{
-		Pass->ResourcesToBegin.Add(Pass);
-		Pass->ResourcesToEnd.Add(Pass);
-	}
+	Pass->ResourcesToBegin.Add(Pass);
+	Pass->ResourcesToEnd.Add(Pass);
 
 	AsyncComputePassCount += EnumHasAnyFlags(PassFlags, ERDGPassFlags::AsyncCompute) ? 1 : 0;
 	RasterPassCount += EnumHasAnyFlags(PassFlags, ERDGPassFlags::Raster) ? 1 : 0;
@@ -2536,7 +2560,7 @@ void FRDGBuilder::SubmitBufferUploads(FRHICommandListBase& RHICmdListUpload, UE:
 			FRDGBuffer* Buffer = UploadedBuffer.Buffer;
 			if (!Buffer->ResourceRHI)
 			{
-				SetPooledBufferRHI(Buffer, AllocatePooledBufferRHI(RHICmdListUpload, Buffer));
+				SetExternalPooledBufferRHI(Buffer, AllocatePooledBufferRHI(RHICmdListUpload, Buffer));
 			}
 		}
 	}
@@ -2746,31 +2770,46 @@ void FRDGBuilder::AllocatePooledTextures(FRHICommandListBase& InRHICmdList, TCon
 
 	for (FCollectResourceOp Op : Ops)
 	{
-		FRDGTexture* Texture = Textures[FRDGTextureHandle(Op.ResourceIndex)];
+		FRDGTexture* Texture = Textures[Op.GetTextureHandle()];
 
-		switch (Op.Op)
+		// External render targets will have the allocation assigned. Scheduled render targets don't yet.
+		check(Texture->Allocation.IsValid() == Texture->bExternal);
+
+		switch (Op.GetOp())
 		{
 		case FCollectResourceOp::EOp::Allocate:
-			SetPooledRenderTargetRHI(Texture, AllocatePooledRenderTargetRHI(InRHICmdList, Texture));
-			break;
+		{
+			FPooledRenderTarget* RenderTarget = GRenderTargetPool.ScheduleAllocation(InRHICmdList, Texture->Desc, Texture->Name, GetAllocateFences(Texture));
+			Texture->RenderTarget = RenderTarget;
+			SetPooledTextureRHI(Texture, &RenderTarget->PooledTexture);
+		}
+		break;
 		case FCollectResourceOp::EOp::Deallocate:
-			if (static_cast<FPooledRenderTarget*>(Texture->RenderTarget)->IsTracked())
+		{
+			FPooledRenderTarget* RenderTarget = static_cast<FPooledRenderTarget*>(Texture->RenderTarget);
+			GRenderTargetPool.ScheduleDeallocation(RenderTarget, GetDeallocateFences(Texture));
+
+			if (Texture->Allocation && RenderTarget->IsTracked())
 			{
 				// This releases the reference without invoking a virtual function call.
 				TRefCountPtr<FPooledRenderTarget>(MoveTemp(Texture->Allocation));
 			}
-			break;
+		}
+		break;
 		}
 	}
 
 	for (FCollectResourceOp Op : Ops)
 	{
-		FRDGTexture* Texture = Textures[FRDGTextureHandle(Op.ResourceIndex)];
+		FRDGTexture* Texture = Textures[Op.GetTextureHandle()];
 
-		if (!Texture->bSkipLastTransition)
+		if (!Texture->bSkipLastTransition && !Texture->Allocation)
 		{
+			FPooledRenderTarget* RenderTarget = static_cast<FPooledRenderTarget*>(Texture->RenderTarget);
+			GRenderTargetPool.FinishSchedule(InRHICmdList, RenderTarget, Texture->Name);
+
 			// Hold the last reference in a chain of pooled allocations.
-			Texture->Allocation = Texture->RenderTarget;
+			Texture->Allocation = RenderTarget;
 		}
 	}
 }
@@ -2782,14 +2821,18 @@ void FRDGBuilder::AllocatePooledBuffers(FRHICommandListBase& InRHICmdList, TCons
 
 	for (FCollectResourceOp Op : Ops)
 	{
-		FRDGBuffer* Buffer = Buffers[FRDGBufferHandle(Op.ResourceIndex)];
+		FRDGBuffer* Buffer = Buffers[Op.GetBufferHandle()];
 
-		switch (Op.Op)
+		switch (Op.GetOp())
 		{
 		case FCollectResourceOp::EOp::Allocate:
-			SetPooledBufferRHI(Buffer, AllocatePooledBufferRHI(InRHICmdList, Buffer));
-			break;
+		{
+			FRDGPooledBuffer* PooledBuffer = GRenderGraphResourcePool.ScheduleAllocation(InRHICmdList, Buffer->Desc, Buffer->Name, ERDGPooledBufferAlignment::Page, GetAllocateFences(Buffer));
+			SetPooledBufferRHI(Buffer, PooledBuffer);
+		}
+		break;
 		case FCollectResourceOp::EOp::Deallocate:
+			GRenderGraphResourcePool.ScheduleDeallocation(Buffer->PooledBuffer, GetDeallocateFences(Buffer));
 			Buffer->Allocation = nullptr;
 			break;
 		}
@@ -2797,10 +2840,12 @@ void FRDGBuilder::AllocatePooledBuffers(FRHICommandListBase& InRHICmdList, TCons
 
 	for (FCollectResourceOp Op : Ops)
 	{
-		FRDGBuffer* Buffer = Buffers[FRDGBufferHandle(Op.ResourceIndex)];
+		FRDGBuffer* Buffer = Buffers[Op.GetBufferHandle()];
 
-		if (!Buffer->bSkipLastTransition)
+		if (!Buffer->bSkipLastTransition && !Buffer->Allocation)
 		{
+			GRenderGraphResourcePool.FinishSchedule(InRHICmdList, Buffer->PooledBuffer);
+
 			// Hold the last reference in a chain of pooled allocations.
 			Buffer->Allocation = Buffer->PooledBuffer;
 		}
@@ -2809,51 +2854,49 @@ void FRDGBuilder::AllocatePooledBuffers(FRHICommandListBase& InRHICmdList, TCons
 
 void FRDGBuilder::AllocateTransientResources(TConstArrayView<FCollectResourceOp> Ops)
 {
+	if (!TransientResourceAllocator)
+	{
+		return;
+	}
+	
 	SCOPED_NAMED_EVENT_TEXT("FRDGBuilder::AllocateTransientResources", FColor::Magenta);
+	TransientResourceAllocator->SetCreateMode(ParallelSetup.bEnabled ? ERHITransientResourceCreateMode::Task : ERHITransientResourceCreateMode::Inline);
+
+	TArray<TPair<FRDGViewableResource*, FRHITransientResource*>, FRDGArrayAllocator> AllocatedResources;
+	AllocatedResources.Reserve(Ops.Num() / 2);
+
 	for (FCollectResourceOp Op : Ops)
 	{
-		const FRDGPassHandle PassHandle = Op.PassHandle;
-
-		switch (Op.Op)
+		switch (Op.GetOp())
 		{
+		default: checkNoEntry();
 		case FCollectResourceOp::EOp::Allocate:
 		{
-			if (Op.ResourceType == ERDGViewableResourceType::Buffer)
+			if (Op.GetResourceType() == ERDGViewableResourceType::Buffer)
 			{
-				FRDGBuffer* Buffer = Buffers[FRDGBufferHandle(Op.ResourceIndex)];
-				FRHITransientBuffer* TransientBuffer = TransientResourceAllocator->CreateBuffer(Translate(Buffer->Desc), Buffer->Name, PassHandle.GetIndex());
-
-				SetTransientBufferRHI(Buffer, TransientBuffer);
-
-				Buffer->MinAcquirePass = FRDGPassHandle(TransientBuffer->GetAcquirePasses().Min);
+				FRDGBuffer* Buffer = Buffers[Op.GetBufferHandle()];
+				FRHITransientBuffer* TransientBuffer = TransientResourceAllocator->CreateBuffer(Translate(Buffer->Desc), Buffer->Name, GetAllocateFences(Buffer));
+				AllocatedResources.Emplace(Buffer, TransientBuffer);
+				Buffer->TransientBuffer = TransientBuffer;
+				Buffer->AcquirePass = FRDGPassHandle(TransientBuffer->GetAcquirePass());
 			}
 			else
 			{
-				FRDGTexture* Texture = Textures[FRDGTextureHandle(Op.ResourceIndex)];
-				FRHITransientTexture* TransientTexture = TransientResourceAllocator->CreateTexture(Texture->Desc, Texture->Name, PassHandle.GetIndex());
-
-				if (Texture->bExternal || Texture->bExtracted)
-				{
-					SetPooledRenderTargetRHI(Texture, GRDGTransientResourceAllocator.AllocateRenderTarget(TransientTexture));
-				}
-				else
-				{
-					SetTransientTextureRHI(Texture, TransientTexture);
-				}
-
-				Texture->MinAcquirePass = FRDGPassHandle(TransientTexture->GetAcquirePasses().Min);
+				FRDGTexture* Texture = Textures[Op.GetTextureHandle()];
+				FRHITransientTexture* TransientTexture = TransientResourceAllocator->CreateTexture(Texture->Desc, Texture->Name, GetAllocateFences(Texture));
+				AllocatedResources.Emplace(Texture, TransientTexture);
+				Texture->TransientTexture = TransientTexture;
+				Texture->AcquirePass = FRDGPassHandle(TransientTexture->GetAcquirePass());
 			}
 		}
 		break;
 		case FCollectResourceOp::EOp::Deallocate:
 		{
-			if (Op.ResourceType == ERDGViewableResourceType::Buffer)
+			if (Op.GetResourceType() == ERDGViewableResourceType::Buffer)
 			{
-				FRDGBuffer* Buffer = Buffers[FRDGBufferHandle(Op.ResourceIndex)];
+				FRDGBuffer* Buffer = Buffers[Op.GetBufferHandle()];
 				FRHITransientBuffer* TransientBuffer = Buffer->TransientBuffer;
-				TransientResourceAllocator->DeallocateMemory(TransientBuffer, PassHandle.GetIndex());
-
-				Buffer->MinDiscardPass = FRDGPassHandle(TransientBuffer->GetDiscardPasses().Min);
+				TransientResourceAllocator->DeallocateMemory(TransientBuffer, GetDeallocateFences(Buffer));
 			}
 			else
 			{
@@ -2866,24 +2909,44 @@ void FRDGBuilder::AllocateTransientResources(TConstArrayView<FCollectResourceOp>
 					if (!Texture->bExtracted)
 					{
 						// This releases the reference without invoking a virtual function call.
-						GRDGTransientResourceAllocator.Release(TRefCountPtr<FRDGTransientRenderTarget>(MoveTemp(Texture->Allocation)), PassHandle);
+						GRDGTransientResourceAllocator.Release(TRefCountPtr<FRDGTransientRenderTarget>(MoveTemp(Texture->Allocation)), GetDeallocateFences(Texture));
 					}
 				}
 				// Texture is using an internal transient texture.
 				else
 				{
-					TransientResourceAllocator->DeallocateMemory(TransientTexture, PassHandle.GetIndex());
-				}
-
-				if (!TransientTexture->IsAcquired())
-				{
-					Texture->MinDiscardPass = FRDGPassHandle(TransientTexture->GetDiscardPasses().Min);
+					TransientResourceAllocator->DeallocateMemory(TransientTexture, GetDeallocateFences(Texture));
 				}
 			}
 		}
 		break;
 		}
 	}
+
+	for (auto [Resource, TransientResource] : AllocatedResources)
+	{
+		TransientResource->Finish(RHICmdList);
+
+		if (Resource->Type == ERDGViewableResourceType::Buffer)
+		{
+			SetTransientBufferRHI(static_cast<FRDGBuffer*>(Resource), static_cast<FRHITransientBuffer*>(TransientResource));
+		}
+		else
+		{
+			check(Resource->Type == ERDGViewableResourceType::Texture);
+			FRDGTexture* Texture = static_cast<FRDGTexture*>(Resource);
+			FRHITransientTexture* TransientTexture = static_cast<FRHITransientTexture*>(TransientResource);
+
+			if (Texture->bExtracted)
+			{
+				SetExternalPooledRenderTargetRHI(Texture, GRDGTransientResourceAllocator.AllocateRenderTarget(TransientTexture));
+			}
+			else
+			{
+				SetTransientTextureRHI(Texture, TransientTexture);
+			}
+		}
+	}
 }
 
 void FRDGBuilder::CreateViews(FRHICommandListBase& InRHICmdList, TConstArrayView<FRDGViewHandle> ViewsToCreate)
@@ -3126,12 +3189,12 @@ void FRDGBuilder::CollectAllocations(FCollectResourceContext& Context, FRDGPass*
 	{
 		for (FRDGPass::FTextureState& PassState : PassToBegin->TextureStates)
 		{
-			CollectAllocateTexture(Context, Pass->Handle, PassState.Texture);
+			CollectAllocateTexture(Context, Pass->Pipeline, Pass->Handle, PassState.Texture);
 		}
 
 		for (FRDGPass::FBufferState& PassState : PassToBegin->BufferStates)
 		{
-			CollectAllocateBuffer(Context, Pass->Handle, PassState.Buffer);
+			CollectAllocateBuffer(Context, Pass->Pipeline, Pass->Handle, PassState.Buffer);
 		}
 
 		if (!IsImmediateMode())
@@ -3144,7 +3207,7 @@ void FRDGBuilder::CollectAllocations(FCollectResourceContext& Context, FRDGPass*
 					BitRef = false;
 				}
 			}
-  
+
 			for (FRDGViewHandle ViewHandle : PassToBegin->Views)
 			{
 				if (auto BitRef = Context.ViewMap[ViewHandle]; BitRef)
@@ -3153,12 +3216,12 @@ void FRDGBuilder::CollectAllocations(FCollectResourceContext& Context, FRDGPass*
 					BitRef = false;
 				}
 			}
-	  	}
-	  	else
-	  	{
+		}
+		else
+		{
 			Context.UniformBuffers = PassToBegin->UniformBuffers;
 			Context.Views = PassToBegin->Views;
-	  	}
+		}
 	}
 }
 
@@ -3168,17 +3231,17 @@ void FRDGBuilder::CollectDeallocations(FCollectResourceContext& Context, FRDGPas
 	{
 		for (FRDGPass::FTextureState& PassState : PassToEnd->TextureStates)
 		{
-			CollectDeallocateTexture(Context, Pass->Handle, PassState.Texture, PassState.ReferenceCount);
+			CollectDeallocateTexture(Context, Pass->Pipeline, Pass->Handle, PassState.Texture, PassState.ReferenceCount);
 		}
 
 		for (FRDGPass::FBufferState& PassState : PassToEnd->BufferStates)
 		{
-			CollectDeallocateBuffer(Context, Pass->Handle, PassState.Buffer, PassState.ReferenceCount);
+			CollectDeallocateBuffer(Context, Pass->Pipeline, Pass->Handle, PassState.Buffer, PassState.ReferenceCount);
 		}
 	}
 }
 
-void FRDGBuilder::CollectAllocateTexture(FCollectResourceContext& Context, FRDGPassHandle PassHandle, FRDGTextureRef Texture)
+void FRDGBuilder::CollectAllocateTexture(FCollectResourceContext& Context, ERHIPipeline PassPipeline, FRDGPassHandle PassHandle, FRDGTextureRef Texture)
 {
 	check(Texture->ReferenceCount > 0 || Texture->bExternal || IsImmediateMode());
 
@@ -3186,9 +3249,6 @@ void FRDGBuilder::CollectAllocateTexture(FCollectResourceContext& Context, FRDGP
 	{
 		FRDGPass* Pass = Passes[PassHandle];
 
-		// Cannot begin a resource on an async compute pass.
-		check(Pass->Pipeline == ERHIPipeline::Graphics);
-
 		// Cannot begin a resource within a merged render pass region.
 		checkf(GetPrologueBarrierPassHandle(PassHandle) == PassHandle,
 			TEXT("Cannot begin a resource within a merged render pass. Pass (Handle: %d, Name: %s), Resource %s"), PassHandle.GetIndex(), Pass->GetName(), Texture->Name);
@@ -3205,15 +3265,15 @@ void FRDGBuilder::CollectAllocateTexture(FCollectResourceContext& Context, FRDGP
 		Texture->bCollectForAllocate = false;
 		check(!Texture->ResourceRHI);
 
-		const FCollectResourceOp AllocateOp = FCollectResourceOp::Allocate(PassHandle, Texture->Handle);
+		const FCollectResourceOp AllocateOp = FCollectResourceOp::Allocate(Texture->Handle);
 
 		if (Texture->bTransient)
 		{
 			Context.TransientResources.Emplace(AllocateOp);
 
-	#if RDG_STATS
+#if RDG_STATS
 			GRDGStatTransientTextureCount++;
-	#endif
+#endif
 		}
 		else
 		{
@@ -3222,17 +3282,18 @@ void FRDGBuilder::CollectAllocateTexture(FCollectResourceContext& Context, FRDGP
 	}
 }
 
-void FRDGBuilder::CollectDeallocateTexture(FCollectResourceContext& Context, FRDGPassHandle PassHandle, FRDGTexture* Texture, uint32 ReferenceCount)
+void FRDGBuilder::CollectDeallocateTexture(FCollectResourceContext& Context, ERHIPipeline PassPipeline, FRDGPassHandle PassHandle, FRDGTexture* Texture, uint32 ReferenceCount)
 {
 	check(!IsImmediateMode());
 	check(Texture->ReferenceCount != FRDGViewableResource::DeallocatedReferenceCount);
 	check(Texture->ReferenceCount >= ReferenceCount);
 	Texture->ReferenceCount -= ReferenceCount;
+	Texture->LastPasses[PassPipeline] = PassHandle;
 
 	if (Texture->ReferenceCount == 0)
 	{
 		check(!Texture->bCollectForAllocate);
-		const FCollectResourceOp DeallocateOp = FCollectResourceOp::Deallocate(PassHandle, Texture->Handle);
+		const FCollectResourceOp DeallocateOp = FCollectResourceOp::Deallocate(Texture->Handle);
 
 		if (Texture->bTransient)
 		{
@@ -3243,12 +3304,11 @@ void FRDGBuilder::CollectDeallocateTexture(FCollectResourceContext& Context, FRD
 			Context.PooledTextures.Emplace(DeallocateOp);
 		}
 
-		Texture->LastPass = PassHandle;
 		Texture->ReferenceCount = FRDGViewableResource::DeallocatedReferenceCount;
 	}
 }
 
-void FRDGBuilder::CollectAllocateBuffer(FCollectResourceContext& Context, FRDGPassHandle PassHandle, FRDGBuffer* Buffer)
+void FRDGBuilder::CollectAllocateBuffer(FCollectResourceContext& Context, ERHIPipeline PassPipeline, FRDGPassHandle PassHandle, FRDGBuffer* Buffer)
 {
 	check(Buffer->ReferenceCount > 0 || IsImmediateMode());
 
@@ -3256,9 +3316,6 @@ void FRDGBuilder::CollectAllocateBuffer(FCollectResourceContext& Context, FRDGPa
 	{
 		const FRDGPass* Pass = Passes[PassHandle];
 
-		// Cannot begin a resource on an async compute pass.
-		check(Pass->Pipeline == ERHIPipeline::Graphics);
-
 		// Cannot begin a resource within a merged render pass region.
 		checkf(GetPrologueBarrierPassHandle(PassHandle) == PassHandle,
 			TEXT("Cannot begin a resource within a merged render pass. Pass (Handle: %d, Name: %s), Resource %s"), PassHandle.GetIndex(), Pass->GetName(), Buffer->Name);
@@ -3275,7 +3332,7 @@ void FRDGBuilder::CollectAllocateBuffer(FCollectResourceContext& Context, FRDGPa
 		Buffer->bCollectForAllocate = false;
 		check(!Buffer->ResourceRHI);
 
-		const FCollectResourceOp AllocateOp = FCollectResourceOp::Allocate(PassHandle, Buffer->Handle);
+		const FCollectResourceOp AllocateOp = FCollectResourceOp::Allocate(Buffer->Handle);
 
 		if (Buffer->bTransient)
 		{
@@ -3292,16 +3349,17 @@ void FRDGBuilder::CollectAllocateBuffer(FCollectResourceContext& Context, FRDGPa
 	}
 }
 
-void FRDGBuilder::CollectDeallocateBuffer(FCollectResourceContext& Context, FRDGPassHandle PassHandle, FRDGBuffer* Buffer, uint32 ReferenceCount)
+void FRDGBuilder::CollectDeallocateBuffer(FCollectResourceContext& Context, ERHIPipeline PassPipeline, FRDGPassHandle PassHandle, FRDGBuffer* Buffer, uint32 ReferenceCount)
 {
 	check(!IsImmediateMode());
 	check(Buffer->ReferenceCount != FRDGViewableResource::DeallocatedReferenceCount);
 	check(Buffer->ReferenceCount >= ReferenceCount);
 	Buffer->ReferenceCount -= ReferenceCount;
+	Buffer->LastPasses[PassPipeline] = PassHandle;
 
 	if (Buffer->ReferenceCount == 0)
 	{
-		const FCollectResourceOp DeallocateOp = FCollectResourceOp::Deallocate(PassHandle, Buffer->Handle);
+		const FCollectResourceOp DeallocateOp = FCollectResourceOp::Deallocate(Buffer->Handle);
 
 		if (Buffer->bTransient)
 		{
@@ -3312,7 +3370,6 @@ void FRDGBuilder::CollectDeallocateBuffer(FCollectResourceContext& Context, FRDG
 			Context.PooledBuffers.Emplace(DeallocateOp);
 		}
 
-		Buffer->LastPass = PassHandle;
 		Buffer->ReferenceCount = FRDGViewableResource::DeallocatedReferenceCount;
 	}
 }
@@ -3598,16 +3655,16 @@ void FRDGBuilder::AddFirstTextureTransition(FRDGTexture* Texture)
 	}
 	else
 	{
-		if (Texture->MinAcquirePass.IsValid())
+		if (Texture->AcquirePass.IsValid())
 		{
-			AddAliasingTransition(Texture->MinAcquirePass, Texture->FirstPass, Texture, FRHITransientAliasingInfo::Acquire(Texture->GetRHI(), Texture->AliasingOverlaps));
+			AddAliasingTransition(Texture->AcquirePass, Texture->FirstPass, Texture, FRHITransientAliasingInfo::Acquire(Texture->GetRHI(), Texture->AliasingOverlaps));
 
-			SubresourceStateBefore.SetPass(ERHIPipeline::Graphics, Texture->MinAcquirePass);
+			SubresourceStateBefore.SetPass(GetPassPipeline(Texture->AcquirePass), Texture->AcquirePass);
 			SubresourceStateBefore.Access = ERHIAccess::Discard;
 		}
 		else if (!Texture->bSplitFirstTransition)
 		{
-			SubresourceStateBefore.SetPass(ERHIPipeline::Graphics, Texture->FirstPass);
+			SubresourceStateBefore.SetPass(GetPassPipeline(Texture->FirstPass), Texture->FirstPass);
 		}
 
 		InitTextureSubresources(*StateBefore, Texture->Layout, &SubresourceStateBefore);
@@ -3630,12 +3687,9 @@ void FRDGBuilder::AddLastTextureTransition(FRDGTexture* Texture)
 	SubresourceStateAfter.SetPass(ERHIPipeline::Graphics, EpiloguePassHandle);
 
 	// Texture is using the RHI transient allocator. Transition it back to Discard in the final pass it is used.
-	if (Texture->MinDiscardPass.IsValid())
+	if (Texture->DiscardPass.IsValid())
 	{
-		FRDGPassHandle MaxDiscardPass = FRDGPassHandle(FMath::Min<uint32>(Texture->TransientTexture->GetDiscardPasses().Max, GetEpiloguePassHandle().GetIndex()));
-		AddAliasingTransition(Texture->MinDiscardPass, MaxDiscardPass, Texture, FRHITransientAliasingInfo::Discard(Texture->GetRHIUnchecked()));
-
-		SubresourceStateAfter.SetPass(ERHIPipeline::Graphics, MaxDiscardPass);
+		SubresourceStateAfter.SetPass(GetPassPipeline(Texture->DiscardPass), Texture->DiscardPass);
 		SubresourceStateAfter.Access = ERHIAccess::Discard;
 	}
 	else
@@ -3654,7 +3708,7 @@ void FRDGBuilder::AddLastTextureTransition(FRDGTexture* Texture)
 			if (!SubresourceStateBefore)
 			{
 				SubresourceStateBefore = AllocSubresource();
-				SubresourceStateBefore->SetPass(ERHIPipeline::Graphics, Texture->FirstPass);
+				SubresourceStateBefore->SetPass(GetPassPipeline(Texture->FirstPass), Texture->FirstPass);
 			}
 
 			State = SubresourceStateBefore;
@@ -3683,16 +3737,16 @@ void FRDGBuilder::AddFirstBufferTransition(FRDGBuffer* Buffer)
 	{
 		StateBefore = AllocSubresource();
 
-		if (Buffer->MinAcquirePass.IsValid())
+		if (Buffer->AcquirePass.IsValid())
 		{
-			AddAliasingTransition(Buffer->MinAcquirePass, Buffer->FirstPass, Buffer, FRHITransientAliasingInfo::Acquire(Buffer->GetRHI(), Buffer->AliasingOverlaps));
+			AddAliasingTransition(Buffer->AcquirePass, Buffer->FirstPass, Buffer, FRHITransientAliasingInfo::Acquire(Buffer->GetRHI(), Buffer->AliasingOverlaps));
 
-			StateBefore->SetPass(ERHIPipeline::Graphics, Buffer->MinAcquirePass);
+			StateBefore->SetPass(GetPassPipeline(Buffer->AcquirePass), Buffer->AcquirePass);
 			StateBefore->Access = ERHIAccess::Discard;
 		}
 		else if (!Buffer->bSplitFirstTransition)
 		{
-			StateBefore->SetPass(ERHIPipeline::Graphics, Buffer->FirstPass);
+			StateBefore->SetPass(GetPassPipeline(Buffer->FirstPass), Buffer->FirstPass);
 		}
 		else
 		{
@@ -3713,12 +3767,10 @@ void FRDGBuilder::AddLastBufferTransition(FRDGBuffer* Buffer)
 	FRDGSubresourceState* StateAfter = AllocSubresource();
 
 	// Texture is using the RHI transient allocator. Transition it back to Discard in the final pass it is used.
-	if (Buffer->MinDiscardPass.IsValid())
+	if (Buffer->DiscardPass.IsValid())
 	{
-		FRDGPassHandle MaxDiscardPass = FRDGPassHandle(FMath::Min<uint32>(Buffer->TransientBuffer->GetDiscardPasses().Max, GetEpiloguePassHandle().GetIndex()));
-		AddAliasingTransition(Buffer->MinDiscardPass, MaxDiscardPass, Buffer, FRHITransientAliasingInfo::Discard(Buffer->GetRHIUnchecked()));
-
-		StateAfter->SetPass(ERHIPipeline::Graphics, MaxDiscardPass);
+		const FRDGPassHandle MaxDiscardPass(FMath::Min<uint32>(Buffer->TransientBuffer->GetDiscardPass(), GetEpiloguePassHandle().GetIndex()));
+		StateAfter->SetPass(GetPassPipeline(MaxDiscardPass), MaxDiscardPass);
 		StateAfter->Access = ERHIAccess::Discard;
 	}
 	else
@@ -4043,6 +4095,72 @@ void FRDGBuilder::AddAliasingTransition(FRDGPassHandle BeginPassHandle, FRDGPass
 	EndPass->GetPrologueBarriersToEnd(Allocators.Transition).AddDependency(BarriersToBegin);
 }
 
+FRHITransientAllocationFences FRDGBuilder::GetAllocateFences(FRDGViewableResource* Resource) const
+{
+	FRHITransientAllocationFences Fences;
+
+	FRDGPassHandle FirstPassHandle = Resource->FirstPass;
+
+	if (GetPassPipeline(FirstPassHandle) == ERHIPipeline::Graphics)
+	{
+		Fences.SetGraphics(FirstPassHandle.GetIndex());
+	}
+	else
+	{
+		const FRDGPass* FirstPass = Passes[FirstPassHandle];
+
+		Fences.SetAsyncCompute(
+			FirstPassHandle.GetIndex(),
+			TInterval<uint32>(FirstPass->GraphicsForkPass.GetIndex(), FirstPass->GraphicsJoinPass.GetIndex()));
+	}
+
+	return Fences;
+}
+
+FRHITransientAllocationFences FRDGBuilder::GetDeallocateFences(FRDGViewableResource* Resource) const
+{
+	FRDGPassHandle GraphicsPassHandle     = Resource->LastPasses[ERHIPipeline::Graphics];
+	FRDGPassHandle AsyncComputePassHandle = Resource->LastPasses[ERHIPipeline::AsyncCompute];
+
+	FRDGPassHandle GraphicsForkPass;
+	FRDGPassHandle GraphicsJoinPass;
+
+	if (AsyncComputePassHandle.IsValid())
+	{
+		const FRDGPass* Pass = Passes[AsyncComputePassHandle];
+		GraphicsForkPass = Pass->GraphicsForkPass;
+		GraphicsJoinPass = Pass->GraphicsJoinPass;
+
+		if (GraphicsPassHandle.IsValid())
+		{
+			// Ignore graphics pass if earlier than the fork to async compute.
+			if (GraphicsPassHandle <= GraphicsForkPass)
+			{
+				GraphicsPassHandle = {};
+			}
+			// Ignore async compute pass if earlier than the join back to graphics.
+			else if (GraphicsPassHandle >= GraphicsJoinPass)
+			{
+				AsyncComputePassHandle = {};
+			}
+		}
+	}
+
+	FRHITransientAllocationFences Fences;
+
+	if (GraphicsPassHandle.IsValid())
+	{
+		Fences.SetGraphics(GraphicsPassHandle.GetIndex());
+	}
+
+	if (AsyncComputePassHandle.IsValid())
+	{
+		Fences.SetAsyncCompute(AsyncComputePassHandle.GetIndex(), TInterval<uint32>(GraphicsForkPass.GetIndex(), GraphicsJoinPass.GetIndex()));
+	}
+
+	return Fences;
+}
+
 ///////////////////////////////////////////////////////////////////////////////////////////////////
 
 TRefCountPtr<IPooledRenderTarget> FRDGBuilder::AllocatePooledRenderTargetRHI(FRHICommandListBase& InRHICmdList, FRDGTextureRef Texture)
@@ -4058,7 +4176,7 @@ TRefCountPtr<FRDGPooledBuffer> FRDGBuilder::AllocatePooledBufferRHI(FRHICommandL
 
 ///////////////////////////////////////////////////////////////////////////////////////////////////
 
-void FRDGBuilder::SetPooledRenderTargetRHI(FRDGTexture* Texture, IPooledRenderTarget* RenderTarget)
+void FRDGBuilder::SetExternalPooledRenderTargetRHI(FRDGTexture* Texture, IPooledRenderTarget* RenderTarget)
 {
 	Texture->RenderTarget = RenderTarget;
 
@@ -4082,9 +4200,7 @@ void FRDGBuilder::SetPooledTextureRHI(FRDGTexture* Texture, FRDGPooledTexture* P
 {
 	check(!Texture->ResourceRHI);
 
-	FRHITexture* TextureRHI = PooledTexture->GetRHI();
-
-	Texture->ResourceRHI = TextureRHI;
+	Texture->SetRHI(PooledTexture->GetRHI());
 	Texture->PooledTexture = PooledTexture;
 	Texture->ViewCache = &PooledTexture->ViewCache;
 
@@ -4103,29 +4219,28 @@ void FRDGBuilder::SetPooledTextureRHI(FRDGTexture* Texture, FRDGPooledTexture* P
 
 void FRDGBuilder::SetTransientTextureRHI(FRDGTexture* Texture, FRHITransientTexture* TransientTexture)
 {
-	check(!Texture->ResourceRHI);
-
-	Texture->ResourceRHI = TransientTexture->GetRHI();
+	Texture->SetRHI(TransientTexture->GetRHI());
 	Texture->TransientTexture = TransientTexture;
 	Texture->ViewCache = &TransientTexture->ViewCache;
-	Texture->bTransient = true;
 	Texture->AliasingOverlaps = TransientTexture->GetAliasingOverlaps();
+
+	if (TransientTexture->IsDiscarded())
+	{
+		Texture->DiscardPass = FRDGPassHandle(FMath::Min<uint32>(TransientTexture->GetDiscardPass(), GetEpiloguePassHandle().GetIndex()));
+	}
 }
 
-void FRDGBuilder::SetPooledBufferRHI(FRDGBuffer* Buffer, FRDGPooledBuffer* PooledBuffer)
+void FRDGBuilder::SetExternalPooledBufferRHI(FRDGBuffer* Buffer, FRDGPooledBuffer* PooledBuffer)
 {
-	check(!Buffer->ResourceRHI);
-
-	FRHIBuffer* BufferRHI = PooledBuffer->GetRHI();
+	SetPooledBufferRHI(Buffer, PooledBuffer);
+	Buffer->Allocation = PooledBuffer;
+}
 
-	Buffer->ResourceRHI = BufferRHI;
+void FRDGBuilder::SetPooledBufferRHI(FRDGBuffer* Buffer, FRDGPooledBuffer* PooledBuffer)
+{
+	Buffer->SetRHI(PooledBuffer->GetRHI());
 	Buffer->PooledBuffer = PooledBuffer;
 	Buffer->ViewCache = &PooledBuffer->ViewCache;
-	Buffer->Allocation = PooledBuffer;
-
-#if RHI_ENABLE_RESOURCE_INFO
-	Buffer->ResourceRHI->SetOwnerName(Buffer->OwnerName);
-#endif
 
 	FRDGBuffer*& Owner = *PooledBufferOwnershipMap.FindOrAdd(PooledBuffer, nullptr);
 
@@ -4142,16 +4257,15 @@ void FRDGBuilder::SetPooledBufferRHI(FRDGBuffer* Buffer, FRDGPooledBuffer* Poole
 
 void FRDGBuilder::SetTransientBufferRHI(FRDGBuffer* Buffer, FRHITransientBuffer* TransientBuffer)
 {
-	check(!Buffer->ResourceRHI && Buffer->bTransient);
-
-	Buffer->ResourceRHI = TransientBuffer->GetRHI();
+	Buffer->SetRHI(TransientBuffer->GetRHI());
 	Buffer->TransientBuffer = TransientBuffer;
 	Buffer->ViewCache = &TransientBuffer->ViewCache;
 	Buffer->AliasingOverlaps = TransientBuffer->GetAliasingOverlaps();
 
-#if RHI_ENABLE_RESOURCE_INFO
-	Buffer->ResourceRHI->SetOwnerName(Buffer->OwnerName);
-#endif
+	if (TransientBuffer->IsDiscarded())
+	{
+		Buffer->DiscardPass = FRDGPassHandle(FMath::Min<uint32>(TransientBuffer->GetDiscardPass(), GetEpiloguePassHandle().GetIndex()));
+	}
 }
 
 ///////////////////////////////////////////////////////////////////////////////////////////////////
@@ -4610,4 +4724,4 @@ void FRDGBuilder::ForceCopyCrossGPU()
 		RHICmdList.TransferResources(Transfers);
 	}
 }
-#endif  // WITH_MGPU
+#endif  // WITH_MGPU
\ No newline at end of file
diff --git a/Engine/Source/Runtime/RenderCore/Private/RenderGraphPrivate.cpp b/Engine/Source/Runtime/RenderCore/Private/RenderGraphPrivate.cpp
--- a/Engine/Source/Runtime/RenderCore/Private/RenderGraphPrivate.cpp
+++ b/Engine/Source/Runtime/RenderCore/Private/RenderGraphPrivate.cpp
@@ -325,6 +325,14 @@ FAutoConsoleVariableRef CVarRDGTransientExtractedResource(
 	TEXT(" 2: force enables all external transient resources (not recommended);"),
 	ECVF_RenderThreadSafe);
 
+int32 GRDGAsyncComputeTransientAliasing = 1;
+FAutoConsoleVariableRef CVarRDGAsyncComputeTransientAliasing(
+	TEXT("r.RDG.AsyncComputeTransientAliasing"), GRDGAsyncComputeTransientAliasing,
+	TEXT("RDG will alias async compute resources on the same heap as graphics resources using fences. This must also be supported by the RHI.")
+	TEXT(" 0: disables transient async compute aliasing;")
+	TEXT(" 1: enables transient async compute aliasing (default);"),
+	ECVF_RenderThreadSafe);
+
 #if RDG_EVENTS
 TAutoConsoleVariable<int32> CVarRDGEvents(
 	TEXT("r.RDG.Events"),
@@ -707,3 +715,12 @@ FRDGScopeState::FState::FState(bool bInImmediate, bool bInParallelExecute)
 	}())
 #endif // RDG_EVENTS
 {}
+
+bool IsExtendedLifetimeResource(FRDGViewableResource* Resource)
+{
+#if RDG_ENABLE_DEBUG
+	return IsDebugAllowedForResource(Resource->Name) && Resource->ReferenceCount != 0 && Resource->ReferenceCount != FRDGViewableResource::DeallocatedReferenceCount;
+#else
+	return false;
+#endif
+}
\ No newline at end of file
diff --git a/Engine/Source/Runtime/RenderCore/Private/RenderGraphPrivate.h b/Engine/Source/Runtime/RenderCore/Private/RenderGraphPrivate.h
--- a/Engine/Source/Runtime/RenderCore/Private/RenderGraphPrivate.h
+++ b/Engine/Source/Runtime/RenderCore/Private/RenderGraphPrivate.h
@@ -114,6 +114,7 @@ extern int32 GRDGAsyncCompute;
 extern int32 GRDGCullPasses;
 extern int32 GRDGMergeRenderPasses;
 extern int32 GRDGTransientAllocator;
+extern int32 GRDGAsyncComputeTransientAliasing;
 extern int32 GRDGTransientExtractedResources;
 extern int32 GRDGTransientIndirectArgBuffers;
 
@@ -254,9 +255,15 @@ FORCEINLINE bool IsAsyncComputeSupported()
 
 extern bool IsParallelExecuteEnabled();
 extern bool IsParallelSetupEnabled();
+extern bool IsExtendedLifetimeResource(FRDGViewableResource* Resource);
+
+inline bool IsAsyncComputeTransientAliasingEnabled()
+{
+	return GRDGAsyncComputeTransientAliasing && GRHIGlobals.SupportsAsyncComputeTransientAliasing;
+}
 
 template <typename ResourceRegistryType, typename FunctionType>
-inline void EnumerateExtendedLifetimeResources(ResourceRegistryType& Registry, FunctionType Function)
+void EnumerateExtendedLifetimeResources(ResourceRegistryType& Registry, FunctionType Function)
 {
 #if RDG_ENABLE_DEBUG
 	if (GRDGDebugExtendResourceLifetimes)
@@ -265,11 +272,11 @@ inline void EnumerateExtendedLifetimeResources(ResourceRegistryType& Registry, F
 		{
 			auto* Resource = Registry[Handle];
 
-			if (IsDebugAllowedForResource(Resource->Name) && !Resource->IsCulled())
+			if (IsExtendedLifetimeResource(Resource))
 			{
 				Function(Resource);
 			}
 		}
 	}
 #endif
-}
+}
\ No newline at end of file
diff --git a/Engine/Source/Runtime/RenderCore/Private/RenderGraphResourcePool.cpp b/Engine/Source/Runtime/RenderCore/Private/RenderGraphResourcePool.cpp
--- a/Engine/Source/Runtime/RenderCore/Private/RenderGraphResourcePool.cpp
+++ b/Engine/Source/Runtime/RenderCore/Private/RenderGraphResourcePool.cpp
@@ -6,6 +6,7 @@
 #include "RHITransientResourceAllocator.h"
 #include "Trace/Trace.inl"
 #include "ProfilingDebugging/CountersTrace.h"
+#include "RenderCore.h"
 
 TRACE_DECLARE_INT_COUNTER(BufferPoolCount, TEXT("BufferPool/BufferCount"));
 TRACE_DECLARE_INT_COUNTER(BufferPoolCreateCount, TEXT("BufferPool/BufferCreateCount"));
@@ -28,6 +29,31 @@ static FAutoConsoleCommandWithOutputDevice GDumpBufferPoolMemoryCmd(
 	FConsoleCommandWithOutputDeviceDelegate::CreateStatic(DumpBufferPoolMemory)
 );
 
+static FRDGBufferDesc GetAlignedBufferDesc(const FRDGBufferDesc& Desc, const TCHAR* InDebugName, ERDGPooledBufferAlignment Alignment)
+{
+	const uint64 BufferPageSize = 64 * 1024;
+
+	FRDGBufferDesc AlignedDesc = Desc;
+
+	switch (Alignment)
+	{
+	case ERDGPooledBufferAlignment::PowerOfTwo:
+		AlignedDesc.NumElements = FMath::RoundUpToPowerOfTwo(AlignedDesc.BytesPerElement * AlignedDesc.NumElements) / AlignedDesc.BytesPerElement;
+		// Fall through to align up to page size for small buffers; helps with reuse.
+
+	case ERDGPooledBufferAlignment::Page:
+		AlignedDesc.NumElements = Align(AlignedDesc.BytesPerElement * AlignedDesc.NumElements, BufferPageSize) / AlignedDesc.BytesPerElement;
+	}
+
+	if (!ensureMsgf(AlignedDesc.NumElements >= Desc.NumElements, TEXT("Alignment caused buffer size overflow for buffer '%s' (AlignedDesc.NumElements: %d < Desc.NumElements: %d)"), InDebugName, AlignedDesc.NumElements, Desc.NumElements))
+	{
+		// Use the unaligned desc since we apparently overflowed when rounding up.
+		AlignedDesc = Desc;
+	}
+
+	return AlignedDesc;
+}
+
 void FRDGBufferPool::DumpMemoryUsage(FOutputDevice& OutputDevice)
 {
 	OutputDevice.Logf(TEXT("Pooled Buffers:"));
@@ -57,118 +83,117 @@ void FRDGBufferPool::DumpMemoryUsage(FOutputDevice& OutputDevice)
 	}
 }
 
-TRefCountPtr<FRDGPooledBuffer> FRDGBufferPool::FindFreeBuffer(FRHICommandListBase& RHICmdList, const FRDGBufferDesc& Desc, const TCHAR* InDebugName, ERDGPooledBufferAlignment Alignment)
+template <typename T>
+FRDGPooledBuffer* FRDGBufferPool::TryFindPooledBuffer(const FRDGBufferDesc& Desc, uint32 DescHash, T&& Predicate)
 {
-	const uint64 BufferPageSize = 64 * 1024;
-
-	FRDGBufferDesc AlignedDesc = Desc;
-
-	switch (Alignment)
-	{
-	case ERDGPooledBufferAlignment::PowerOfTwo:
-		AlignedDesc.NumElements = FMath::RoundUpToPowerOfTwo(AlignedDesc.BytesPerElement * AlignedDesc.NumElements) / AlignedDesc.BytesPerElement;
-		// Fall through to align up to page size for small buffers; helps with reuse.
-
-	case ERDGPooledBufferAlignment::Page:
-		AlignedDesc.NumElements = Align(AlignedDesc.BytesPerElement * AlignedDesc.NumElements, BufferPageSize) / AlignedDesc.BytesPerElement;
-	}
-
-	if (!ensureMsgf(AlignedDesc.NumElements >= Desc.NumElements, TEXT("Alignment caused buffer size overflow for buffer '%s' (AlignedDesc.NumElements: %d < Desc.NumElements: %d)"), InDebugName, AlignedDesc.NumElements, Desc.NumElements))
-	{
-		// Use the unaligned desc since we apparently overflowed when rounding up.
-		AlignedDesc = Desc;
-	}
-
-	const uint32 BufferHash = GetTypeHash(AlignedDesc);
-
-	UE::TScopeLock Lock(Mutex);
-
-	// First find if available.
 	for (int32 Index = 0; Index < AllocatedBufferHashes.Num(); ++Index)
 	{
-		if (AllocatedBufferHashes[Index] != BufferHash)
+		if (AllocatedBufferHashes[Index] != DescHash)
 		{
 			continue;
 		}
 
-		const auto& PooledBuffer = AllocatedBuffers[Index];
+		FRDGPooledBuffer* Found = AllocatedBuffers[Index];
 
 		// Still being used outside the pool.
-		if (PooledBuffer->GetRefCount() > 1)
+		if (Found->GetRefCount() > 1 || !Predicate(Found))
 		{
 			continue;
 		}
 
-		check(PooledBuffer->GetAlignedDesc() == AlignedDesc);
+		check(Found->GetAlignedDesc() == Desc);
+		return Found;
+	}
+	return nullptr;
+}
 
-		PooledBuffer->LastUsedFrame = FrameCounter;
+FRDGPooledBuffer* FRDGBufferPool::ScheduleAllocation(
+	FRHICommandListBase& RHICmdList,
+	const FRDGBufferDesc& Desc,
+	const TCHAR* Name,
+	ERDGPooledBufferAlignment Alignment,
+	const FRHITransientAllocationFences& Fences)
+{
+	const FRDGBufferDesc AlignedDesc = GetAlignedBufferDesc(Desc, Name, Alignment);
+	const uint32 DescHash = GetTypeHash(AlignedDesc);
 
-#if (UE_BUILD_SHIPPING || UE_BUILD_TEST)
-		PooledBuffer->Name = InDebugName;
-#else
-		// For performance, avoid updating name if it happens to be the same (true 80% of the time in testing)
-		bool bNameUpdated = false;
-		if (!PooledBuffer->Name || FCString::Strcmp(PooledBuffer->Name, InDebugName))
-		{
-			// Name changed, need to update it
-			PooledBuffer->Name = InDebugName;
+	FRDGPooledBuffer* PooledBuffer = TryFindPooledBuffer(AlignedDesc, DescHash, [&](FRDGPooledBuffer* PooledBuffer)
+	{
+		return PooledBuffer->Fences && !FRHITransientAllocationFences::Contains(*PooledBuffer->Fences, Fences);
+	});
 
-			RHICmdList.BindDebugLabelName(PooledBuffer->GetRHI(), InDebugName);
+	if (!PooledBuffer)
+	{
+		PooledBuffer = CreateBuffer(RHICmdList, AlignedDesc, DescHash, Name);
+	}
 
-			bNameUpdated = true;
-		}
+	// We need the external-facing desc to match what the user requested.
+	const_cast<FRDGBufferDesc&>(PooledBuffer->Desc).NumElements = Desc.NumElements;
+	PooledBuffer->Fences.Reset();
+	PooledBuffer->LastUsedFrame = FrameCounter;
+	return PooledBuffer;
+}
 
-		// Propagate the debug name to ViewCache if the name was updated, or if any items were added to ViewCache since the debug name was set
-		int32 ViewCacheNum = PooledBuffer->ViewCache.NumItems();
-		if (bNameUpdated || PooledBuffer->NameUpdatedViewCacheNum != ViewCacheNum)
-		{
-			PooledBuffer->NameUpdatedViewCacheNum = ViewCacheNum;
-			PooledBuffer->ViewCache.SetDebugName(RHICmdList, InDebugName);
-		}
-#endif
+void FRDGBufferPool::ScheduleDeallocation(FRDGPooledBuffer* PooledBuffer, const FRHITransientAllocationFences& Fences)
+{
+	PooledBuffer->Fences = Fences;
+}
 
-		// We need the external-facing desc to match what the user requested.
-		const_cast<FRDGBufferDesc&>(PooledBuffer->Desc).NumElements = Desc.NumElements;
+void FRDGBufferPool::FinishSchedule(FRHICommandListBase& RHICmdList, FRDGPooledBuffer* PooledBuffer)
+{
+	PooledBuffer->Fences.Emplace();
+	PooledBuffer->SetDebugLabelName(RHICmdList, PooledBuffer->Name);
+}
 
-		return PooledBuffer;
-	}
+TRefCountPtr<FRDGPooledBuffer> FRDGBufferPool::FindFreeBuffer(FRHICommandListBase& RHICmdList, const FRDGBufferDesc& Desc, const TCHAR* InDebugName, ERDGPooledBufferAlignment Alignment)
+{
+	const FRDGBufferDesc AlignedDesc = GetAlignedBufferDesc(Desc, InDebugName, Alignment);
+	const uint32 DescHash = GetTypeHash(AlignedDesc);
 
-	// Allocate new one
-	{
-		const uint32 NumBytes = AlignedDesc.GetSize();
+	UE::TScopeLock Lock(Mutex);
 
-#if CPUPROFILERTRACE_ENABLED
-		UE_TRACE_LOG_SCOPED_T(Cpu, FRDGBufferPool_CreateBuffer, CpuChannel)
-			<< FRDGBufferPool_CreateBuffer.Name(InDebugName)
-			<< FRDGBufferPool_CreateBuffer.SizeInBytes(NumBytes);
-#endif
+	FRDGPooledBuffer* PooledBuffer = TryFindPooledBuffer(AlignedDesc, DescHash);
+
+	if (!PooledBuffer)
+	{
+		PooledBuffer = CreateBuffer(RHICmdList, AlignedDesc, DescHash, InDebugName);
+	}
 
-		TRACE_COUNTER_ADD(BufferPoolCount, 1);
-		TRACE_COUNTER_ADD(BufferPoolCreateCount, 1);
-		TRACE_COUNTER_ADD(BufferPoolSize, NumBytes);
+	// We need the external-facing desc to match what the user requested.
+	const_cast<FRDGBufferDesc&>(PooledBuffer->Desc).NumElements = Desc.NumElements;
+	PooledBuffer->LastUsedFrame = FrameCounter;
+	PooledBuffer->SetDebugLabelName(RHICmdList, InDebugName);
+	return PooledBuffer;
+}
 
-		const ERHIAccess InitialAccess = RHIGetDefaultResourceState(Desc.Usage, false);
-		FRHIResourceCreateInfo CreateInfo(InDebugName);
-		TRefCountPtr<FRHIBuffer> BufferRHI = RHICmdList.CreateBuffer(NumBytes, Desc.Usage, Desc.BytesPerElement, InitialAccess, CreateInfo);
+FRDGPooledBuffer* FRDGBufferPool::CreateBuffer(FRHICommandListBase& RHICmdList, const FRDGBufferDesc& Desc, uint32 DescHash, const TCHAR* InDebugName)
+{
+	const uint32 NumBytes = Desc.GetSize();
 
-	#if !(UE_BUILD_SHIPPING || UE_BUILD_TEST)
-		RHICmdList.BindDebugLabelName(BufferRHI, InDebugName);
-	#endif
+#if CPUPROFILERTRACE_ENABLED
+	UE_TRACE_LOG_SCOPED_T(Cpu, FRDGBufferPool_CreateBuffer, CpuChannel)
+		<< FRDGBufferPool_CreateBuffer.Name(InDebugName)
+		<< FRDGBufferPool_CreateBuffer.SizeInBytes(NumBytes);
+#endif
 
-		TRefCountPtr<FRDGPooledBuffer> PooledBuffer = new FRDGPooledBuffer(RHICmdList, MoveTemp(BufferRHI), Desc, AlignedDesc.NumElements, InDebugName);
-		AllocatedBuffers.Add(PooledBuffer);
-		AllocatedBufferHashes.Add(BufferHash);
-		check(PooledBuffer->GetRefCount() == 2);
+	TRACE_COUNTER_ADD(BufferPoolCount, 1);
+	TRACE_COUNTER_ADD(BufferPoolCreateCount, 1);
+	TRACE_COUNTER_ADD(BufferPoolSize, NumBytes);
 
-		PooledBuffer->LastUsedFrame = FrameCounter;
+	const ERHIAccess InitialAccess = RHIGetDefaultResourceState(Desc.Usage, false);
+	FRHIResourceCreateInfo CreateInfo(InDebugName);
+	TRefCountPtr<FRHIBuffer> BufferRHI = RHICmdList.CreateBuffer(NumBytes, Desc.Usage, Desc.BytesPerElement, InitialAccess, CreateInfo);
 
-		if (EnumHasAllFlags(Desc.Usage, EBufferUsageFlags::ReservedResource))
-		{
-			PooledBuffer->CommittedSizeInBytes = 0;
-		}
+	FRDGPooledBuffer* PooledBuffer = new FRDGPooledBuffer(RHICmdList, MoveTemp(BufferRHI), Desc, Desc.NumElements, InDebugName);
+	AllocatedBuffers.Add(PooledBuffer);
+	AllocatedBufferHashes.Add(DescHash);
 
-		return PooledBuffer;
+	if (EnumHasAllFlags(Desc.Usage, EBufferUsageFlags::ReservedResource))
+	{
+		PooledBuffer->CommittedSizeInBytes = 0;
 	}
+
+	return PooledBuffer;
 }
 
 void FRDGBufferPool::ReleaseRHI()
@@ -299,7 +324,7 @@ TRefCountPtr<FRDGTransientRenderTarget> FRDGTransientResourceAllocator::Allocate
 	return RenderTarget;
 }
 
-void FRDGTransientResourceAllocator::Release(TRefCountPtr<FRDGTransientRenderTarget>&& RenderTarget, FRDGPassHandle PassHandle)
+void FRDGTransientResourceAllocator::Release(TRefCountPtr<FRDGTransientRenderTarget>&& RenderTarget, const FRHITransientAllocationFences& Fences)
 {
 	check(RenderTarget);
 
@@ -307,7 +332,7 @@ void FRDGTransientResourceAllocator::Release(TRefCountPtr<FRDGTransientRenderTar
 	// members before dereferencing to zero so that it gets marked as deallocated rather than pending.
 	if (RenderTarget->GetRefCount() == 1)
 	{
-		Allocator->DeallocateMemory(RenderTarget->Texture, PassHandle.GetIndex());
+		Allocator->DeallocateMemory(RenderTarget->Texture, Fences);
 		RenderTarget->Reset();
 		RenderTarget = nullptr;
 	}
@@ -346,9 +371,9 @@ void FRDGTransientResourceAllocator::ReleasePendingDeallocations()
 
 		for (FRDGTransientRenderTarget* RenderTarget : PendingDeallocationList)
 		{
-			Allocator->DeallocateMemory(RenderTarget->Texture, 0);
-
-			Aliases.Emplace(FRHITransientAliasingInfo::Discard(RenderTarget->Texture->GetRHI()));
+			FRHITransientAllocationFences Fences;
+			Fences.SetGraphics(0);
+			Allocator->DeallocateMemory(RenderTarget->Texture, Fences);
 			Transitions.Emplace(RenderTarget->Texture->GetRHI(), ERHIAccess::Unknown, ERHIAccess::Discard);
 
 			RenderTarget->Reset();
diff --git a/Engine/Source/Runtime/RenderCore/Private/RenderGraphResourcePool.h b/Engine/Source/Runtime/RenderCore/Private/RenderGraphResourcePool.h
--- a/Engine/Source/Runtime/RenderCore/Private/RenderGraphResourcePool.h
+++ b/Engine/Source/Runtime/RenderCore/Private/RenderGraphResourcePool.h
@@ -31,6 +31,27 @@ public:
 private:
 	RENDERCORE_API void ReleaseRHI() override;
 
+	FRDGPooledBuffer* CreateBuffer(FRHICommandListBase& RHICmdList, const FRDGBufferDesc& Desc, uint32 DescHash, const TCHAR* InDebugName);
+
+	FRDGPooledBuffer* ScheduleAllocation(
+		FRHICommandListBase& RHICmdList,
+		const FRDGBufferDesc& Desc,
+		const TCHAR* Name,
+		ERDGPooledBufferAlignment Alignment,
+		const FRHITransientAllocationFences& Fences);
+
+	void ScheduleDeallocation(FRDGPooledBuffer* PooledBuffer, const FRHITransientAllocationFences& Fences);
+
+	void FinishSchedule(FRHICommandListBase& RHICmdList, FRDGPooledBuffer* PooledBuffer);
+
+	template <typename T>
+	FRDGPooledBuffer* TryFindPooledBuffer(const FRDGBufferDesc& Desc, uint32 DescHash, T&& Predicate);
+
+	FRDGPooledBuffer* TryFindPooledBuffer(const FRDGBufferDesc& Desc, uint32 DescHash)
+	{
+		return TryFindPooledBuffer(Desc, DescHash, [](FRDGPooledBuffer*) { return true; });
+	}
+
 	mutable UE::FRecursiveMutex Mutex;
 
 	/** Elements can be 0, we compact the buffer later. */
@@ -96,7 +117,7 @@ public:
 
 	TRefCountPtr<FRDGTransientRenderTarget> AllocateRenderTarget(FRHITransientTexture* Texture);
 
-	void Release(TRefCountPtr<FRDGTransientRenderTarget>&& RenderTarget, FRDGPassHandle PassHandle);
+	void Release(TRefCountPtr<FRDGTransientRenderTarget>&& RenderTarget, const FRHITransientAllocationFences& Fences);
 
 	void ReleasePendingDeallocations();
 
diff --git a/Engine/Source/Runtime/RenderCore/Private/RenderGraphResources.cpp b/Engine/Source/Runtime/RenderCore/Private/RenderGraphResources.cpp
--- a/Engine/Source/Runtime/RenderCore/Private/RenderGraphResources.cpp
+++ b/Engine/Source/Runtime/RenderCore/Private/RenderGraphResources.cpp
@@ -146,6 +146,33 @@ FRDGPooledBuffer::FRDGPooledBuffer(TRefCountPtr<FRHIBuffer> InBuffer, const FRDG
 	: FRDGPooledBuffer(FRHICommandListImmediate::Get(), MoveTemp(InBuffer), InDesc, InNumAllocatedElements, InName)
 {}
 
+void FRDGPooledBuffer::SetDebugLabelName(FRHICommandListBase& RHICmdList, const TCHAR* InName)
+{
+#if (UE_BUILD_SHIPPING || UE_BUILD_TEST)
+	Name = InName;
+#else
+	// For performance, avoid updating name if it happens to be the same (true 80% of the time in testing)
+	bool bNameUpdated = false;
+	if (!InName || FCString::Strcmp(Name, InName))
+	{
+		// Name changed, need to update it
+		Name = InName;
+
+		RHICmdList.BindDebugLabelName(GetRHI(), InName);
+
+		bNameUpdated = true;
+	}
+
+	// Propagate the debug name to ViewCache if the name was updated, or if any items were added to ViewCache since the debug name was set
+	int32 ViewCacheNum = ViewCache.NumItems();
+	if (bNameUpdated || NameUpdatedViewCacheNum != ViewCacheNum)
+	{
+		NameUpdatedViewCacheNum = ViewCacheNum;
+		ViewCache.SetDebugName(RHICmdList, InName);
+	}
+#endif
+}
+
 void FRDGUniformBuffer::InitRHI()
 {
 	check(!HasRHI());
diff --git a/Engine/Source/Runtime/RenderCore/Private/RenderGraphTrace.cpp b/Engine/Source/Runtime/RenderCore/Private/RenderGraphTrace.cpp
--- a/Engine/Source/Runtime/RenderCore/Private/RenderGraphTrace.cpp
+++ b/Engine/Source/Runtime/RenderCore/Private/RenderGraphTrace.cpp
@@ -57,6 +57,8 @@ UE_TRACE_EVENT_BEGIN(RDGTrace, BufferMessage)
 	UE_TRACE_EVENT_FIELD(uint64[], TransientAllocationOffsetMins)
 	UE_TRACE_EVENT_FIELD(uint64[], TransientAllocationOffsetMaxs)
 	UE_TRACE_EVENT_FIELD(uint16[], TransientAllocationMemoryRanges)
+	UE_TRACE_EVENT_FIELD(FRDGPassHandle::IndexType, TransientAcquirePass)
+	UE_TRACE_EVENT_FIELD(FRDGPassHandle::IndexType, TransientDiscardPass)
 	UE_TRACE_EVENT_FIELD(bool, IsExternal)
 	UE_TRACE_EVENT_FIELD(bool, IsExtracted)
 	UE_TRACE_EVENT_FIELD(bool, IsCulled)
@@ -78,6 +80,8 @@ UE_TRACE_EVENT_BEGIN(RDGTrace, TextureMessage)
 	UE_TRACE_EVENT_FIELD(uint64[], TransientAllocationOffsetMins)
 	UE_TRACE_EVENT_FIELD(uint64[], TransientAllocationOffsetMaxs)
 	UE_TRACE_EVENT_FIELD(uint16[], TransientAllocationMemoryRanges)
+	UE_TRACE_EVENT_FIELD(FRDGPassHandle::IndexType, TransientAcquirePass)
+	UE_TRACE_EVENT_FIELD(FRDGPassHandle::IndexType, TransientDiscardPass)
 	UE_TRACE_EVENT_FIELD(uint64, SizeInBytes)
 	UE_TRACE_EVENT_FIELD(uint64, CreateFlags)
 	UE_TRACE_EVENT_FIELD(uint32, Dimension)
@@ -276,6 +280,36 @@ void FRDGTrace::OutputGraphEnd(const FRDGBuilder& GraphBuilder)
 
 	} TransientAllocation;
 
+	FRDGPassHandle TransientAcquirePass;
+	FRDGPassHandle TransientDiscardPass;
+
+	const auto FillTransientResourceArrays = [&](const FRHITransientResource* Resource, bool bRemoveFromStats)
+	{
+		TransientAllocation.Reset();
+		TransientAcquirePass = {};
+		TransientDiscardPass = {};
+
+		if (Resource)
+		{
+			TransientAllocation.Fill(TransientAllocationStats, Resource);
+
+			if (Resource->GetAcquirePass() != FRHITransientResource::kInvalidPassIndex)
+			{
+				TransientAcquirePass = FRDGPassHandle(Resource->GetAcquirePass());
+			}
+
+			if (Resource->GetDiscardPass() != FRHITransientResource::kInvalidPassIndex)
+			{
+				TransientDiscardPass = FRDGPassHandle(Resource->GetDiscardPass());
+			}
+
+			if (bRemoveFromStats)
+			{
+				TransientAllocationStats.Resources.Remove(Resource);
+			}
+		}
+	};
+
 	for (FRDGTextureHandle Handle = Textures.Begin(); Handle != Textures.End(); ++Handle)
 	{
 		const FRDGTexture* Texture = Textures[Handle];
@@ -283,17 +317,18 @@ void FRDGTrace::OutputGraphEnd(const FRDGBuilder& GraphBuilder)
 		uint64 SizeInBytes = 0;
 		if (FRHITexture* TextureRHI = Texture->GetRHIUnchecked())
 		{
-			SizeInBytes = RHIComputeMemorySize(TextureRHI);
+			if (Texture->TransientTexture)
+			{
+				SizeInBytes = Texture->TransientTexture->GetSize();
+			}
+			else
+			{
+				SizeInBytes = RHIComputeMemorySize(TextureRHI);
+			}
 		}
 
-		TransientAllocation.Reset();
-
-		if (Texture->TransientTexture)
-		{
-			const FRHITransientResource* Resource = Texture->TransientTexture;
-			TransientAllocation.Fill(TransientAllocationStats, Resource);
-			TransientAllocationStats.Resources.Remove(Resource);
-		}
+		const bool bRemoveFromStats = true;
+		FillTransientResourceArrays(Texture->TransientTexture, bRemoveFromStats);
 
 		UE_TRACE_LOG(RDGTrace, TextureMessage, RDGChannel)
 			<< TextureMessage.Name(Texture->Name, uint16(FCString::Strlen(Texture->Name)))
@@ -304,6 +339,8 @@ void FRDGTrace::OutputGraphEnd(const FRDGBuilder& GraphBuilder)
 			<< TextureMessage.TransientAllocationOffsetMins(TransientAllocation.OffsetMins.GetData(), TransientAllocation.OffsetMins.Num())
 			<< TextureMessage.TransientAllocationOffsetMaxs(TransientAllocation.OffsetMaxs.GetData(), TransientAllocation.OffsetMaxs.Num())
 			<< TextureMessage.TransientAllocationMemoryRanges(TransientAllocation.MemoryRanges.GetData(), TransientAllocation.MemoryRanges.Num())
+			<< TextureMessage.TransientAcquirePass(TransientAcquirePass.GetIndexUnchecked())
+			<< TextureMessage.TransientDiscardPass(TransientDiscardPass.GetIndexUnchecked())
 			<< TextureMessage.SizeInBytes(SizeInBytes)
 			<< TextureMessage.CreateFlags(uint32(Texture->Desc.Flags))
 			<< TextureMessage.Dimension(uint32(Texture->Desc.Dimension))
@@ -327,15 +364,9 @@ void FRDGTrace::OutputGraphEnd(const FRDGBuilder& GraphBuilder)
 	for (FRDGBufferHandle Handle = Buffers.Begin(); Handle != Buffers.End(); ++Handle)
 	{
 		const FRDGBuffer* Buffer = Buffers[Handle];
-
-		TransientAllocation.Reset();
-
-		if (Buffer->TransientBuffer)
-		{
-			const FRHITransientResource* Resource = Buffer->TransientBuffer;
-			TransientAllocation.Fill(TransientAllocationStats, Resource);
-			TransientAllocationStats.Resources.Remove(Resource);
-		}
+		
+		const bool bRemoveFromStats = true;
+		FillTransientResourceArrays(Buffer->TransientBuffer, bRemoveFromStats);
 
 		UE_TRACE_LOG(RDGTrace, BufferMessage, RDGChannel)
 			<< BufferMessage.Name(Buffer->Name, uint16(FCString::Strlen(Buffer->Name)))
@@ -346,6 +377,8 @@ void FRDGTrace::OutputGraphEnd(const FRDGBuilder& GraphBuilder)
 			<< BufferMessage.TransientAllocationOffsetMins(TransientAllocation.OffsetMins.GetData(), TransientAllocation.OffsetMins.Num())
 			<< BufferMessage.TransientAllocationOffsetMaxs(TransientAllocation.OffsetMaxs.GetData(), TransientAllocation.OffsetMaxs.Num())
 			<< BufferMessage.TransientAllocationMemoryRanges(TransientAllocation.MemoryRanges.GetData(), TransientAllocation.MemoryRanges.Num())
+			<< BufferMessage.TransientAcquirePass(TransientAcquirePass.GetIndexUnchecked())
+			<< BufferMessage.TransientDiscardPass(TransientDiscardPass.GetIndexUnchecked())
 			<< BufferMessage.UsageFlags(uint32(Buffer->Desc.Usage))
 			<< BufferMessage.BytesPerElement(Buffer->Desc.BytesPerElement)
 			<< BufferMessage.NumElements(Buffer->Desc.NumElements)
@@ -371,8 +404,8 @@ void FRDGTrace::OutputGraphEnd(const FRDGBuilder& GraphBuilder)
 			continue;
 		}
 
-		TransientAllocation.Reset();
-		TransientAllocation.Fill(TransientAllocationStats, Resource);
+		const bool bRemoveFromStats = false;
+		FillTransientResourceArrays(Resource, bRemoveFromStats);
 
 		if (Resource->GetResourceType() == ERHITransientResourceType::Texture)
 		{
@@ -384,6 +417,8 @@ void FRDGTrace::OutputGraphEnd(const FRDGBuilder& GraphBuilder)
 				<< TextureMessage.TransientAllocationOffsetMins(TransientAllocation.OffsetMins.GetData(), TransientAllocation.OffsetMins.Num())
 				<< TextureMessage.TransientAllocationOffsetMaxs(TransientAllocation.OffsetMaxs.GetData(), TransientAllocation.OffsetMaxs.Num())
 				<< TextureMessage.TransientAllocationMemoryRanges(TransientAllocation.MemoryRanges.GetData(), TransientAllocation.MemoryRanges.Num())
+				<< TextureMessage.TransientAcquirePass(TransientAcquirePass.GetIndexUnchecked())
+				<< TextureMessage.TransientDiscardPass(TransientDiscardPass.GetIndexUnchecked())
 				<< TextureMessage.SizeInBytes(Resource->GetSize())
 				<< TextureMessage.CreateFlags(uint32(Texture->CreateInfo.Flags))
 				<< TextureMessage.Dimension(uint32(Texture->CreateInfo.Dimension))
@@ -415,6 +450,8 @@ void FRDGTrace::OutputGraphEnd(const FRDGBuilder& GraphBuilder)
 				<< BufferMessage.TransientAllocationOffsetMins(TransientAllocation.OffsetMins.GetData(), TransientAllocation.OffsetMins.Num())
 				<< BufferMessage.TransientAllocationOffsetMaxs(TransientAllocation.OffsetMaxs.GetData(), TransientAllocation.OffsetMaxs.Num())
 				<< BufferMessage.TransientAllocationMemoryRanges(TransientAllocation.MemoryRanges.GetData(), TransientAllocation.MemoryRanges.Num())
+				<< BufferMessage.TransientAcquirePass(TransientAcquirePass.GetIndexUnchecked())
+				<< BufferMessage.TransientDiscardPass(TransientDiscardPass.GetIndexUnchecked())
 				<< BufferMessage.UsageFlags(uint32(Buffer->CreateInfo.Usage))
 				<< BufferMessage.BytesPerElement(Buffer->CreateInfo.Stride)
 				<< BufferMessage.NumElements(Buffer->CreateInfo.Size / Buffer->CreateInfo.Stride)
diff --git a/Engine/Source/Runtime/RenderCore/Private/RenderGraphValidation.cpp b/Engine/Source/Runtime/RenderCore/Private/RenderGraphValidation.cpp
--- a/Engine/Source/Runtime/RenderCore/Private/RenderGraphValidation.cpp
+++ b/Engine/Source/Runtime/RenderCore/Private/RenderGraphValidation.cpp
@@ -1508,21 +1508,6 @@ void FRDGBarrierValidation::ValidateBarrierBatchEnd(const FRDGPass* Pass, const
 			}
 		}
 
-		for (const auto& KeyValue : ResourceMap.Aliases)
-		{
-			const FRHITransientAliasingInfo& Info = KeyValue.Value;
-			if (Info.IsDiscard())
-			{
-				FRDGViewableResource* Resource = KeyValue.Key;
-
-				if (IsDebugAllowedForResource(Resource->Name))
-				{
-					LogHeader();
-					UE_LOG(LogRDG, Display, TEXT("\tRDG(%p) RHI(%p) %s - Discard"), Resource, Resource->GetRHIUnchecked(), Resource->Name);
-				}
-			}
-		}
-
 		bFoundFirst = false;
 	}
 }
diff --git a/Engine/Source/Runtime/RenderCore/Private/RenderTargetPool.cpp b/Engine/Source/Runtime/RenderCore/Private/RenderTargetPool.cpp
--- a/Engine/Source/Runtime/RenderCore/Private/RenderTargetPool.cpp
+++ b/Engine/Source/Runtime/RenderCore/Private/RenderTargetPool.cpp
@@ -65,107 +65,129 @@ static uint32 ComputeSizeInKB(FPooledRenderTarget& Element)
 	return (Element.ComputeMemorySize() + 1023) / 1024;
 }
 
-TRefCountPtr<IPooledRenderTarget> FRenderTargetPool::FindFreeElement(FRHICommandListBase& RHICmdList, FRHITextureCreateInfo Desc, const TCHAR* Name)
+FPooledRenderTarget* FRenderTargetPool::CreateRenderTarget(FRHICommandListBase& RHICmdList, const FRHITextureCreateInfo& Desc, uint32 DescHash, const TCHAR* Name)
 {
-	FPooledRenderTarget* Found = 0;
-	uint32 FoundIndex = -1;
+#if CPUPROFILERTRACE_ENABLED
+	UE_TRACE_LOG_SCOPED_T(Cpu, FRenderTargetPool_CreateTexture, CpuChannel)
+		<< FRenderTargetPool_CreateTexture.Name(Name);
+#endif
 
-	// FastVRAM is no longer supported by the render target pool.
-	EnumRemoveFlags(Desc.Flags, ETextureCreateFlags::FastVRAM | ETextureCreateFlags::FastVRAMPartialAlloc);
+	const ERHIAccess AccessInitial = ERHIAccess::SRVMask;
+	FRHITextureCreateDesc CreateDesc(Desc, AccessInitial, Name);
+	const static FLazyName ClassName(TEXT("FPooledRenderTarget"));
+	CreateDesc.SetClassName(ClassName);
 
-	// We always want SRV access
-	Desc.Flags |= TexCreate_ShaderResource;
+	FPooledRenderTarget* Result = new FPooledRenderTarget(
+		RHICmdList.CreateTexture(CreateDesc),
+		Translate(CreateDesc),
+		this);
 
-	// Render target pool always forces textures into non streaming memory.
-	// UE-TODO: UE-188415 fix flags that we force in Render Target Pool
-	//Desc.Flags |= ETextureCreateFlags::ForceIntoNonStreamingMemoryTracking;
+	PooledRenderTargets.Add(Result);
+	PooledRenderTargetHashes.Add(DescHash);
 
-	const uint32 DescHash = GetTypeHash(Desc);
+	if (EnumHasAnyFlags(Desc.Flags, TexCreate_UAV))
+	{
+		// The render target desc is invalid if a UAV is requested with an RHI that doesn't support the high-end feature level.
+		check(GMaxRHIFeatureLevel >= ERHIFeatureLevel::SM5 || GMaxRHIFeatureLevel == ERHIFeatureLevel::ES3_1);
 
-	UE::TScopeLock Lock(Mutex);
+		if (GRHISupportsUAVFormatAliasing)
+		{
+			EPixelFormat AliasFormat = Desc.UAVFormat != PF_Unknown
+				? Desc.UAVFormat
+				: Desc.Format;
+
+			Result->RenderTargetItem.UAV = RHICmdList.CreateUnorderedAccessView(Result->GetRHI(), 0, (uint8)AliasFormat, 0, 0);
+		}
+		else
+		{
+			checkf(Desc.UAVFormat == PF_Unknown || Desc.UAVFormat == Desc.Format, TEXT("UAV aliasing is not supported by the current RHI."));
+			Result->RenderTargetItem.UAV = RHICmdList.CreateUnorderedAccessView(Result->GetRHI(), 0);
+		}
+	}
+
+	AllocationLevelInKB += ComputeSizeInKB(*Result);
+	TRACE_COUNTER_ADD(RenderTargetPoolCount, 1);
+	TRACE_COUNTER_SET(RenderTargetPoolSize, (int64)AllocationLevelInKB * 1024);
+	return Result;
+}
 
+template <typename T>
+FPooledRenderTarget* FRenderTargetPool::TryFindRenderTarget(const FRHITextureCreateInfo& Desc, uint32 DescHash, T&& Predicate) const
+{
 	for (uint32 Index = 0, Num = (uint32)PooledRenderTargets.Num(); Index < Num; ++Index)
 	{
 		if (PooledRenderTargetHashes[Index] == DescHash)
 		{
 			FPooledRenderTarget* Element = PooledRenderTargets[Index];
 
-		#if DO_CHECK
-			{
-				checkf(Element, TEXT("Hash was not cleared from the list."));
+			checkf(Element, TEXT("Hash was not cleared from the list."));
+			checkf(Translate(Element->GetDesc()) == Desc, TEXT("Invalid hash or collision when attempting to allocate %s"), Element->GetDesc().DebugName);
 
-				const FRHITextureCreateInfo ElementDesc = Translate(Element->GetDesc());
-				checkf(ElementDesc == Desc, TEXT("Invalid hash or collision when attempting to allocate %s"), Element->GetDesc().DebugName);
-			}
-		#endif
-
-			if (Element->IsFree())
+			if (Element->IsFree() && Predicate(Element))
 			{
-				Found = Element;
-				FoundIndex = Index;
-				break;
+				return Element;
 			}
 		}
 	}
+	return nullptr;
+}
+
+FPooledRenderTarget* FRenderTargetPool::ScheduleAllocation(FRHICommandListBase& RHICmdList, FRHITextureCreateInfo Desc, const TCHAR* Name, const FRHITransientAllocationFences& Fences)
+{
+	// FastVRAM is no longer supported by the render target pool.
+	EnumRemoveFlags(Desc.Flags, ETextureCreateFlags::FastVRAM | ETextureCreateFlags::FastVRAMPartialAlloc);
+
+	// We always want SRV access
+	Desc.Flags |= TexCreate_ShaderResource;
+
+	const uint32 DescHash = GetTypeHash(Desc);
+	
+	FPooledRenderTarget* Found = TryFindRenderTarget(Desc, DescHash, [&](FPooledRenderTarget* Element)
+	{
+		return Element->PooledTexture.Fences && !FRHITransientAllocationFences::Contains(*Element->PooledTexture.Fences, Fences);
+	});
 
 	if (!Found)
 	{
-#if CPUPROFILERTRACE_ENABLED
-		UE_TRACE_LOG_SCOPED_T(Cpu, FRenderTargetPool_CreateTexture, CpuChannel)
-			<< FRenderTargetPool_CreateTexture.Name(Name);
-#endif
+		Found = CreateRenderTarget(RHICmdList, Desc, DescHash, Name);
+	}
 
-		const ERHIAccess AccessInitial = ERHIAccess::SRVMask;
-		FRHITextureCreateDesc CreateDesc(Desc, AccessInitial, Name);
-		const static FLazyName ClassName(TEXT("FPooledRenderTarget"));
-		CreateDesc.SetClassName(ClassName);
+	Found->PooledTexture.Fences.Reset();
+	Found->UnusedForNFrames = 0;
+	return Found;
+}
 
-		Found = new FPooledRenderTarget(
-			RHICmdList.CreateTexture(CreateDesc),
-			Translate(CreateDesc),
-			this);
+void FRenderTargetPool::ScheduleDeallocation(FPooledRenderTarget* RenderTarget, const FRHITransientAllocationFences& Fences)
+{
+	RenderTarget->PooledTexture.Fences = Fences;
+}
 
-		PooledRenderTargets.Add(Found);
-		PooledRenderTargetHashes.Add(DescHash);
+void FRenderTargetPool::FinishSchedule(FRHICommandListBase& RHICmdList, FPooledRenderTarget* RenderTarget, const TCHAR* Name)
+{
+	RenderTarget->PooledTexture.Fences.Emplace();
+	RenderTarget->SetDebugLabelName(RHICmdList, Name);
+}
 
-		if (EnumHasAnyFlags(Desc.Flags, TexCreate_UAV))
-		{
-			// The render target desc is invalid if a UAV is requested with an RHI that doesn't support the high-end feature level.
-			check(GMaxRHIFeatureLevel >= ERHIFeatureLevel::SM5 || GMaxRHIFeatureLevel == ERHIFeatureLevel::ES3_1);
+TRefCountPtr<IPooledRenderTarget> FRenderTargetPool::FindFreeElement(FRHICommandListBase& RHICmdList, FRHITextureCreateInfo Desc, const TCHAR* Name)
+{
+	// FastVRAM is no longer supported by the render target pool.
+	EnumRemoveFlags(Desc.Flags, ETextureCreateFlags::FastVRAM | ETextureCreateFlags::FastVRAMPartialAlloc);
 
-			if (GRHISupportsUAVFormatAliasing)
-			{
-				EPixelFormat AliasFormat = Desc.UAVFormat != PF_Unknown
-					? Desc.UAVFormat
-					: Desc.Format;
+	// We always want SRV access
+	Desc.Flags |= TexCreate_ShaderResource;
 
-				Found->RenderTargetItem.UAV = RHICmdList.CreateUnorderedAccessView(Found->GetRHI(), 0, (uint8)AliasFormat, 0, 0);
-			}
-			else
-			{
-				checkf(Desc.UAVFormat == PF_Unknown || Desc.UAVFormat == Desc.Format, TEXT("UAV aliasing is not supported by the current RHI."));
-				Found->RenderTargetItem.UAV = RHICmdList.CreateUnorderedAccessView(Found->GetRHI(), 0);
-			}
-		}
+	const uint32 DescHash = GetTypeHash(Desc);
 
-		AllocationLevelInKB += ComputeSizeInKB(*Found);
-		TRACE_COUNTER_ADD(RenderTargetPoolCount, 1);
-		TRACE_COUNTER_SET(RenderTargetPoolSize, (int64)AllocationLevelInKB * 1024);
+	UE::TScopeLock Lock(Mutex);
 
-		FoundIndex = PooledRenderTargets.Num() - 1;
-	}
+	FPooledRenderTarget* Found = TryFindRenderTarget(Desc, DescHash);
 
-#if (UE_BUILD_SHIPPING || UE_BUILD_TEST)
-	Found->Desc.DebugName = Name;
-#else
-	// For performance, avoid updating name if it happens to be the same (true 90% of the time in testing)
-	if (FCString::Strcmp(Found->Desc.DebugName, Name))
+	if (!Found)
 	{
-		Found->Desc.DebugName = Name;
-		RHICmdList.BindDebugLabelName(Found->GetRHI(), Name);
+		Found = CreateRenderTarget(RHICmdList, Desc, DescHash, Name);
 	}
-#endif
 
+	Found->SetDebugLabelName(RHICmdList, Name);
 	Found->UnusedForNFrames = 0;
 
 	return TRefCountPtr<IPooledRenderTarget>(MoveTemp(Found));
@@ -581,6 +603,20 @@ bool FPooledRenderTarget::OnFrameStart()
 	return false;
 }
 
+void FPooledRenderTarget::SetDebugLabelName(FRHICommandListBase& RHICmdList, const TCHAR* Name)
+{
+#if (UE_BUILD_SHIPPING || UE_BUILD_TEST)
+	Desc.DebugName = Name;
+#else
+	// For performance, avoid updating name if it happens to be the same (true 90% of the time in testing)
+	if (FCString::Strcmp(Desc.DebugName, Name))
+	{
+		Desc.DebugName = Name;
+		RHICmdList.BindDebugLabelName(GetRHI(), Name);
+	}
+#endif
+}
+
 uint32 FPooledRenderTarget::ComputeMemorySize() const
 {
 	uint32 Size = 0;
diff --git a/Engine/Source/Runtime/RenderCore/Public/RenderGraphBuilder.h b/Engine/Source/Runtime/RenderCore/Public/RenderGraphBuilder.h
--- a/Engine/Source/Runtime/RenderCore/Public/RenderGraphBuilder.h
+++ b/Engine/Source/Runtime/RenderCore/Public/RenderGraphBuilder.h
@@ -582,38 +582,58 @@ private:
 			Deallocate
 		};
 
-		static FCollectResourceOp Allocate(FRDGPassHandle PassHandle, FRDGBufferHandle BufferHandle)
+		static FCollectResourceOp Allocate(FRDGBufferHandle BufferHandle)
 		{
-			return FCollectResourceOp(PassHandle, BufferHandle.GetIndex(), ERDGViewableResourceType::Buffer, EOp::Allocate);
+			return FCollectResourceOp(BufferHandle.GetIndex(), ERDGViewableResourceType::Buffer, EOp::Allocate);
 		}
 
-		static FCollectResourceOp Allocate(FRDGPassHandle PassHandle, FRDGTextureHandle TextureHandle)
+		static FCollectResourceOp Allocate(FRDGTextureHandle TextureHandle)
 		{
-			return FCollectResourceOp(PassHandle, TextureHandle.GetIndex(), ERDGViewableResourceType::Texture, EOp::Allocate);
+			return FCollectResourceOp(TextureHandle.GetIndex(), ERDGViewableResourceType::Texture, EOp::Allocate);
 		}
 
-		static FCollectResourceOp Deallocate(FRDGPassHandle PassHandle, FRDGBufferHandle BufferHandle)
+		static FCollectResourceOp Deallocate(FRDGBufferHandle BufferHandle)
 		{
-			return FCollectResourceOp(PassHandle, BufferHandle.GetIndex(), ERDGViewableResourceType::Buffer, EOp::Deallocate);
+			return FCollectResourceOp(BufferHandle.GetIndex(), ERDGViewableResourceType::Buffer, EOp::Deallocate);
 		}
 
-		static FCollectResourceOp Deallocate(FRDGPassHandle PassHandle, FRDGTextureHandle TextureHandle)
+		static FCollectResourceOp Deallocate(FRDGTextureHandle TextureHandle)
 		{
-			return FCollectResourceOp(PassHandle, TextureHandle.GetIndex(), ERDGViewableResourceType::Texture, EOp::Deallocate);
+			return FCollectResourceOp(TextureHandle.GetIndex(), ERDGViewableResourceType::Texture, EOp::Deallocate);
 		}
 
 		FCollectResourceOp() = default;
-		FCollectResourceOp(FRDGPassHandle InPassHandle, uint16 InResourceIndex, ERDGViewableResourceType InResourceType, EOp InOp)
-			: PassHandle(InPassHandle)
-			, ResourceIndex(InResourceIndex)
-			, ResourceType(InResourceType)
-			, Op(InOp)
+		FCollectResourceOp(uint32 InResourceIndex, ERDGViewableResourceType InResourceType, EOp InOp)
+			: ResourceIndex(InResourceIndex)
+			, ResourceType(static_cast<uint32>(InResourceType))
+			, Op(static_cast<uint32>(InOp))
 		{}
 
-		FRDGPassHandle PassHandle;
-		uint16 ResourceIndex;
-		ERDGViewableResourceType ResourceType;
-		EOp Op;
+		EOp GetOp() const
+		{
+			return static_cast<EOp>(Op);
+		}
+
+		ERDGViewableResourceType GetResourceType() const
+		{
+			return static_cast<ERDGViewableResourceType>(ResourceType);
+		}
+
+		FRDGTextureHandle GetTextureHandle() const
+		{
+			check(GetResourceType() == ERDGViewableResourceType::Texture);
+			return FRDGTextureHandle(ResourceIndex);
+		}
+
+		FRDGBufferHandle GetBufferHandle() const
+		{
+			check(GetResourceType() == ERDGViewableResourceType::Buffer);
+			return FRDGBufferHandle(ResourceIndex);
+		}
+
+		uint32 ResourceIndex : 30;
+		uint32 ResourceType : 1;
+		uint32 Op : 1;
 	};
 
 	using FCollectResourceOpArray = TArray<FCollectResourceOp, FRDGArrayAllocator>;
@@ -636,13 +656,13 @@ private:
 
 	/** Collects new resource allocations for the pass into the provided context. */
 	void CollectAllocations(FCollectResourceContext& Context, FRDGPass* Pass);
-	void CollectAllocateTexture(FCollectResourceContext& Context, FRDGPassHandle PassHandle, FRDGTexture* Texture);
-	void CollectAllocateBuffer(FCollectResourceContext& Context, FRDGPassHandle PassHandle, FRDGBuffer* Buffer);
+	void CollectAllocateTexture(FCollectResourceContext& Context, ERHIPipeline PassPipeline, FRDGPassHandle PassHandle, FRDGTexture* Texture);
+	void CollectAllocateBuffer(FCollectResourceContext& Context, ERHIPipeline PassPipeline, FRDGPassHandle PassHandle, FRDGBuffer* Buffer);
 
 	/** Collects new resource deallocations for the pass into the provided context. */
 	void CollectDeallocations(FCollectResourceContext& Context, FRDGPass* Pass);
-	void CollectDeallocateTexture(FCollectResourceContext& Context, FRDGPassHandle PassHandle, FRDGTexture* Texture, uint32 ReferenceCount);
-	void CollectDeallocateBuffer(FCollectResourceContext& Context, FRDGPassHandle PassHandle, FRDGBuffer* Buffer, uint32 ReferenceCount);
+	void CollectDeallocateTexture(FCollectResourceContext& Context, ERHIPipeline PassPipeline, FRDGPassHandle PassHandle, FRDGTexture* Texture, uint32 ReferenceCount);
+	void CollectDeallocateBuffer(FCollectResourceContext& Context, ERHIPipeline PassPipeline, FRDGPassHandle PassHandle, FRDGBuffer* Buffer, uint32 ReferenceCount);
 
 	/** Allocates resources using the provided lifetime op arrays. */
 	void AllocateTransientResources(TConstArrayView<FCollectResourceOp> Ops);
@@ -658,9 +678,10 @@ private:
 	TRefCountPtr<FRDGPooledBuffer> AllocatePooledBufferRHI(FRHICommandListBase& RHICmdList, FRDGBufferRef Buffer);
 
 	/** Assigns an underlying RHI resource to an RDG resource. */
-	void SetPooledRenderTargetRHI(FRDGTexture* Texture, IPooledRenderTarget* RenderTarget);
+	void SetExternalPooledRenderTargetRHI(FRDGTexture* Texture, IPooledRenderTarget* RenderTarget);
 	void SetPooledTextureRHI(FRDGTexture* Texture, FRDGPooledTexture* PooledTexture);
 	void SetTransientTextureRHI(FRDGTexture* Texture, FRHITransientTexture* TransientTexture);
+	void SetExternalPooledBufferRHI(FRDGBuffer* Buffer, FRDGPooledBuffer* PooledBuffer);
 	void SetPooledBufferRHI(FRDGBuffer* Buffer, FRDGPooledBuffer* PooledBuffer);
 	void SetTransientBufferRHI(FRDGBuffer* Buffer, FRHITransientBuffer* TransientBuffer);
 
@@ -807,6 +828,17 @@ private:
 		Pass->GetEpilogueBarriersToEnd(Allocators.Transition).AddDependency(&BarriersToBegin);
 	}
 
+	// Returns fences representing an allocation event, which can only happen on one pipeline at a time.
+	FRHITransientAllocationFences GetAllocateFences(FRDGViewableResource* Resource) const;
+
+	// Returns fences representing a deallocation event, which can happen on multiple pipes.
+	FRHITransientAllocationFences GetDeallocateFences(FRDGViewableResource* Resource) const;
+
+	inline ERHIPipeline GetPassPipeline(FRDGPassHandle PassHandle) const
+	{
+		return Passes[PassHandle]->Pipeline;
+	}
+
 	FRDGSubresourceState* AllocSubresource(const FRDGSubresourceState& Other);
 	FRDGSubresourceState* AllocSubresource();
 
diff --git a/Engine/Source/Runtime/RenderCore/Public/RenderGraphDefinitions.h b/Engine/Source/Runtime/RenderCore/Public/RenderGraphDefinitions.h
--- a/Engine/Source/Runtime/RenderCore/Public/RenderGraphDefinitions.h
+++ b/Engine/Source/Runtime/RenderCore/Public/RenderGraphDefinitions.h
@@ -297,7 +297,6 @@ public:
 	FORCEINLINE IndexType GetIndexUnchecked() const { return Index; }
 	FORCEINLINE bool IsNull()  const { return Index == kNullIndex; }
 	FORCEINLINE bool IsValid() const { return Index != kNullIndex; }
-	FORCEINLINE operator bool() const { return IsValid(); }
 	FORCEINLINE bool operator==(TRDGHandle Other) const { return Index == Other.Index; }
 	FORCEINLINE bool operator!=(TRDGHandle Other) const { return Index != Other.Index; }
 	FORCEINLINE bool operator<=(TRDGHandle Other) const { check(IsValid() && Other.IsValid()); return Index <= Other.Index; }
diff --git a/Engine/Source/Runtime/RenderCore/Public/RenderGraphResources.h b/Engine/Source/Runtime/RenderCore/Public/RenderGraphResources.h
--- a/Engine/Source/Runtime/RenderCore/Public/RenderGraphResources.h
+++ b/Engine/Source/Runtime/RenderCore/Public/RenderGraphResources.h
@@ -47,6 +47,7 @@ struct FRDGProducerState
 {
 	FRDGPass* Pass = nullptr;
 	FRDGPass* PassIfSkipUAVBarrier = nullptr;
+	FRDGPass* PassIfReadAccess = nullptr;
 	ERHIAccess Access = ERHIAccess::Unknown;
 	FRDGViewHandle NoUAVBarrierHandle;
 };
@@ -142,13 +143,6 @@ public:
 		return ResourceRHI;
 	}
 
-	void SetOwnerName(const FName& InOwnerName)
-	{
-#if RHI_ENABLE_RESOURCE_INFO
-		OwnerName = InOwnerName;
-#endif
-	}
-
 	//////////////////////////////////////////////////////////////////////////
 
 protected:
@@ -178,10 +172,6 @@ private:
 	RENDERCORE_API FRDGResourceDebugData& GetDebugData() const;
 #endif
 
-#if RHI_ENABLE_RESOURCE_INFO
-	FName OwnerName;	// For RHI resource tracking
-#endif
-
 	friend FRDGBuilder;
 	friend FRDGUserValidation;
 	friend FRDGBarrierValidation;
@@ -317,6 +307,13 @@ public:
 		return bProduced;
 	}
 
+	void SetOwnerName(const FName& InOwnerName)
+	{
+#if RHI_ENABLE_RESOURCE_INFO
+		OwnerName = InOwnerName;
+#endif
+	}
+
 protected:
 	RENDERCORE_API FRDGViewableResource(const TCHAR* InName, ERDGViewableResourceType InType, bool bSkipTracking, bool bImmediateFirstBarrier);
 
@@ -395,10 +392,10 @@ protected:
 	/** Whether this resource is allowed to be both transient and extracted. */
 	ETransientExtractionHint TransientExtractionHint;
 
+	FRDGPassHandle AcquirePass;
+	FRDGPassHandle DiscardPass;
 	FRDGPassHandle FirstPass;
-	FRDGPassHandle LastPass;
-	FRDGPassHandle MinAcquirePass;
-	FRDGPassHandle MinDiscardPass;
+	FRDGPassHandlesByPipeline LastPasses;
 
 	/** Number of references in passes and deferred queries. */
 	uint32 ReferenceCount;
@@ -415,6 +412,16 @@ protected:
 private:
 	static const uint16 DeallocatedReferenceCount = ~0;
 
+	void SetRHI(FRHIResource* Resource)
+	{
+		check(!ResourceRHI);
+		ResourceRHI = Resource;
+
+	#if RHI_ENABLE_RESOURCE_INFO
+		ResourceRHI->SetOwnerName(OwnerName);
+	#endif
+	}
+
 	void SetExternalAccessMode(ERHIAccess InAccess, ERHIPipeline InPipelines)
 	{
 		check(!AccessModeState.bLocked);
@@ -426,6 +433,10 @@ private:
 		EpilogueAccess = InAccess;
 	}
 
+#if RHI_ENABLE_RESOURCE_INFO
+	FName OwnerName;	// For RHI resource tracking
+#endif
+
 #if RDG_ENABLE_TRACE
 	uint32 TraceOrder = 0;
 	TArray<FRDGPassHandle, FRDGArrayAllocator> TracePasses;
@@ -436,6 +447,8 @@ private:
 	RENDERCORE_API FRDGViewableResourceDebugData& GetViewableDebugData() const;
 #endif
 
+	friend bool IsExtendedLifetimeResource(FRDGViewableResource*);
+
 	friend FRDGBuilder;
 	friend FRDGUserValidation;
 	friend FRDGBarrierBatchBegin;
@@ -492,7 +505,9 @@ class FRDGPooledTexture final
 public:
 	FRDGPooledTexture(FRHITexture* InTexture)
 		: Texture(InTexture)
-	{}
+	{
+		Fences.Emplace();
+	}
 
 	/** Finds a UAV matching the descriptor in the cache or creates a new one and updates the cache. */
 	FORCEINLINE FRHIUnorderedAccessView* GetOrCreateUAV(FRHICommandListBase& RHICmdList, const FRHITextureUAVCreateInfo& UAVDesc) { return ViewCache.GetOrCreateUAV(RHICmdList, Texture, UAVDesc); }
@@ -511,8 +526,10 @@ public:
 private:
 	TRefCountPtr<FRHITexture> Texture;
 	FRHITextureViewCache ViewCache;
+	TOptional<FRHITransientAllocationFences> Fences;
 
 	friend FRDGBuilder;
+	friend FRenderTargetPool;
 };
 
 /** Render graph tracked Texture. */
@@ -1153,6 +1170,8 @@ public:
 		{
 			CachedSRV = GetOrCreateSRV(RHICmdList, FRHIBufferSRVCreateInfo());
 		}
+
+		Fences.Emplace();
 	}
 
 	RENDERCORE_API FRDGPooledBuffer(TRefCountPtr<FRHIBuffer> InBuffer, const FRDGBufferDesc& InDesc, uint32 InNumAllocatedElements, const TCHAR* InName);
@@ -1224,6 +1243,8 @@ private:
 		return AlignedDesc;
 	}
 
+	void SetDebugLabelName(FRHICommandListBase& RHICmdList, const TCHAR* InName);
+
 	// Used internally by FRDGBuilder::QueueCommitReservedBuffer(),
 	// which is expected to be the only way to resize physical memory for FRDGPooledBuffer
 	void SetCommittedSize(uint64 InCommittedSizeInBytes)
@@ -1253,6 +1274,8 @@ private:
 	const uint32 NumAllocatedElements;
 	uint32 LastUsedFrame = 0;
 
+	TOptional<FRHITransientAllocationFences> Fences;
+
 	friend FRDGBuilder;
 	friend FRDGBufferPool;
 };
diff --git a/Engine/Source/Runtime/RenderCore/Public/RenderTargetPool.h b/Engine/Source/Runtime/RenderCore/Public/RenderTargetPool.h
--- a/Engine/Source/Runtime/RenderCore/Public/RenderTargetPool.h
+++ b/Engine/Source/Runtime/RenderCore/Public/RenderTargetPool.h
@@ -65,6 +65,8 @@ struct FPooledRenderTarget final : public IPooledRenderTarget
 	RENDERCORE_API uint32 ComputeMemorySize() const override;
 
 private:
+	RENDERCORE_API void SetDebugLabelName(FRHICommandListBase& RHICmdList, const TCHAR* Name);
+
 	/** Pointer back to the pool for render targets which are actually pooled, otherwise NULL. */
 	FRenderTargetPool* RenderTargetPool;
 	
@@ -156,7 +158,32 @@ public:
 	RENDERCORE_API void DumpMemoryUsage(FOutputDevice& OutputDevice);
 
 private:
-	RENDERCORE_API void FreeElementAtIndex(int32 Index);
+	void FreeElementAtIndex(int32 Index);
+
+	FPooledRenderTarget* CreateRenderTarget(FRHICommandListBase& RHICmdList, const FRHITextureCreateInfo& Desc, uint32 DescHash, const TCHAR* Name);
+
+	template <typename T>
+	FPooledRenderTarget* TryFindRenderTarget(const FRHITextureCreateInfo& Desc, uint32 DescHash, T&& Predicate) const;
+
+	FPooledRenderTarget* TryFindRenderTarget(const FRHITextureCreateInfo& Desc, uint32 DescHash) const
+	{
+		return TryFindRenderTarget(Desc, DescHash, [](FPooledRenderTarget*) { return true; });
+	}
+
+	//////////////////////////////////////////////////////////////////////////
+	// Methods for scheduling allocations for RDG
+
+	FPooledRenderTarget* ScheduleAllocation(
+		FRHICommandListBase& RHICmdList,
+		FRHITextureCreateInfo Desc,
+		const TCHAR* Name,
+		const FRHITransientAllocationFences& Fences);
+
+	void ScheduleDeallocation(FPooledRenderTarget* RenderTarget, const FRHITransientAllocationFences& Fences);
+
+	void FinishSchedule(FRHICommandListBase& RHICmdList, FPooledRenderTarget* RenderTarget, const TCHAR* Name);
+
+	//////////////////////////////////////////////////////////////////////////
 
 	mutable UE::FRecursiveMutex Mutex;
 
@@ -172,7 +199,7 @@ private:
 	bool bCurrentlyOverBudget = false;
 
 	// could be done on the fly but that makes the RenderTargetPoolEvents harder to read
-	RENDERCORE_API void CompactPool();
+	void CompactPool();
 
 	friend struct FPooledRenderTarget;
 	friend class FVisualizeTexture;
diff --git a/Engine/Source/Runtime/VulkanRHI/Private/VulkanTransientResourceAllocator.cpp b/Engine/Source/Runtime/VulkanRHI/Private/VulkanTransientResourceAllocator.cpp
--- a/Engine/Source/Runtime/VulkanRHI/Private/VulkanTransientResourceAllocator.cpp
+++ b/Engine/Source/Runtime/VulkanRHI/Private/VulkanTransientResourceAllocator.cpp
@@ -94,11 +94,11 @@ FVulkanTransientResourceAllocator::FVulkanTransientResourceAllocator(FVulkanTran
 {
 }
 
-FRHITransientTexture* FVulkanTransientResourceAllocator::CreateTexture(const FRHITextureCreateInfo& InCreateInfo, const TCHAR* InDebugName, uint32 InPassIndex)
+FRHITransientTexture* FVulkanTransientResourceAllocator::CreateTexture(const FRHITextureCreateInfo& InCreateInfo, const TCHAR* InDebugName, const FRHITransientAllocationFences& Fences)
 {
 	FDynamicRHI::FRHICalcTextureSizeResult MemReq = GDynamicRHI->RHICalcTexturePlatformSize(InCreateInfo, 0);
 
-	return CreateTextureInternal(InCreateInfo, InDebugName, InPassIndex, MemReq.Size, MemReq.Align,
+	return CreateTextureInternal(InCreateInfo, InDebugName, Fences, MemReq.Size, MemReq.Align,
 		[&](const FRHITransientHeap::FResourceInitializer& Initializer)
 	{
 		FRHITextureCreateDesc CreateDesc(InCreateInfo, ERHIAccess::Discard, InDebugName);
@@ -107,7 +107,7 @@ FRHITransientTexture* FVulkanTransientResourceAllocator::CreateTexture(const FRH
 	});
 }
 
-FRHITransientBuffer* FVulkanTransientResourceAllocator::CreateBuffer(const FRHIBufferCreateInfo& InCreateInfo, const TCHAR* InDebugName, uint32 InPassIndex)
+FRHITransientBuffer* FVulkanTransientResourceAllocator::CreateBuffer(const FRHIBufferCreateInfo& InCreateInfo, const TCHAR* InDebugName, const FRHITransientAllocationFences& Fences)
 {
 	checkf(!EnumHasAnyFlags(InCreateInfo.Usage, BUF_AccelerationStructure), TEXT("AccelerationStructure not yet supported as TransientResource."));
 	checkf(!EnumHasAnyFlags(InCreateInfo.Usage, BUF_Volatile), TEXT("The volatile flag is not supported for transient resources."));
@@ -116,7 +116,7 @@ FRHITransientBuffer* FVulkanTransientResourceAllocator::CreateBuffer(const FRHIB
 	const uint32 Alignment = VulkanRHI::FMemoryManager::CalculateBufferAlignment(*Device, InCreateInfo.Usage, bZeroSize);
 	uint64 Size = Align(InCreateInfo.Size, Alignment);
 
-	return CreateBufferInternal(InCreateInfo, InDebugName, InPassIndex, Size, Alignment,
+	return CreateBufferInternal(InCreateInfo, InDebugName, Fences, Size, Alignment,
 		[&](const FRHITransientHeap::FResourceInitializer& Initializer)
 	{
 		FRHIResourceCreateInfo ResourceCreateInfo(InDebugName);
diff --git a/Engine/Source/Runtime/VulkanRHI/Private/VulkanTransientResourceAllocator.h b/Engine/Source/Runtime/VulkanRHI/Private/VulkanTransientResourceAllocator.h
--- a/Engine/Source/Runtime/VulkanRHI/Private/VulkanTransientResourceAllocator.h
+++ b/Engine/Source/Runtime/VulkanRHI/Private/VulkanTransientResourceAllocator.h
@@ -56,6 +56,6 @@ public:
 		default: checkNoEntry(); return false;
 		}
 	}
-	FRHITransientTexture* CreateTexture(const FRHITextureCreateInfo& InCreateInfo, const TCHAR* InDebugName, uint32 InPassIndex) override;
-	FRHITransientBuffer* CreateBuffer(const FRHIBufferCreateInfo& InCreateInfo, const TCHAR* InDebugName, uint32 InPassIndex) override;
+	FRHITransientTexture* CreateTexture(const FRHITextureCreateInfo& InCreateInfo, const TCHAR* InDebugName, const FRHITransientAllocationFences& Fences) override;
+	FRHITransientBuffer* CreateBuffer(const FRHIBufferCreateInfo& InCreateInfo, const TCHAR* InDebugName, const FRHITransientAllocationFences& Fences) override;
 };
\ No newline at end of file
