PR #11821: Fix uniform buffer hash mismatch log print order.

diff --git a/Engine/Source/Runtime/RHICore/Private/RHICoreShader.cpp b/Engine/Source/Runtime/RHICore/Private/RHICoreShader.cpp
--- a/Engine/Source/Runtime/RHICore/Private/RHICoreShader.cpp
+++ b/Engine/Source/Runtime/RHICore/Private/RHICoreShader.cpp
@@ -57,7 +57,7 @@ void ValidateStaticUniformBuffer(FRHIUniformBuffer* UniformBuffer, FUniformBuffe
 			checkf(
 				false,
 				TEXT("Shader attempted to bind uniform buffer '%s' at slot %s with hash '%u', but the shader expected '%s' with hash '%u'."),
-				*Layout.GetDebugName(), *SlotRegistry.GetDebugDescription(Slot), ExpectedHash, ExpectedStructMetadata->GetShaderVariableName(), Layout.GetHash());
+				*Layout.GetDebugName(), *SlotRegistry.GetDebugDescription(Slot), Layout.GetHash(), ExpectedStructMetadata->GetShaderVariableName(), ExpectedHash);
 		}
 	}
 #endif
@@ -238,4 +238,4 @@ const FString* GetDiagnosticMessage(uint32 MessageID)
 }
 
 } //! RHICore
-} //! UE
\ No newline at end of file
+} //! UE
