Aftermath: Add option to hook ShaderDebugInfo callback in Aftermath to dump nvdbg file (useful for SPIRV shader line in Vulkan crashes).

#rnx

diff --git a/Engine/Source/Runtime/RHICore/Private/RHICoreNvidiaAftermath.cpp b/Engine/Source/Runtime/RHICore/Private/RHICoreNvidiaAftermath.cpp
--- a/Engine/Source/Runtime/RHICore/Private/RHICoreNvidiaAftermath.cpp
+++ b/Engine/Source/Runtime/RHICore/Private/RHICoreNvidiaAftermath.cpp
@@ -65,6 +65,13 @@ namespace UE::RHICore::Nvidia::Aftermath
 		ECVF_ReadOnly
 	);
 
+	static TAutoConsoleVariable<float> CVarAftermath_DumpShaderDebugInfo(
+		TEXT("r.GPUCrashDebugging.Aftermath.DumpShaderDebugInfo"),
+		0,
+		TEXT("Dump shader debug info (.nvdbg) alongside the crash dump."),
+		ECVF_ReadOnly
+	);
+
 	static TAutoConsoleVariable<float> CVar_DumpWaitTime(
 		TEXT("r.GPUCrashDebugging.Aftermath.DumpWaitTime"),
 		10.0f,
@@ -91,6 +98,8 @@ namespace UE::RHICore::Nvidia::Aftermath
 	static void GFSDK_AFTERMATH_CALL Callback_GpuCrashDump    (const void* GPUCrashDumpData, const uint32_t GPUCrashDumpSize, void* UserData);
 	static void GFSDK_AFTERMATH_CALL Callback_GpuCrashDumpDesc(PFN_GFSDK_Aftermath_AddGpuCrashDumpDescription AddValue, void* UserData);
 	static void GFSDK_AFTERMATH_CALL Callback_ResolveMarker   (const void* MarkerData, const uint32_t MarkerDataSize, void* UserData, void** ResolvedMarkerData, uint32_t* ResolvedMarkerDataSize);
+	static void GFSDK_AFTERMATH_CALL Callback_ShaderDebugInfo (const void* ShaderDebugInfo, const uint32 ShaderDebugInfoSize, void* UserData);
+
 
 	void StartupModule()
 	{
@@ -140,6 +149,8 @@ namespace UE::RHICore::Nvidia::Aftermath
 		Flags |= (FParse::Param(FCommandLine::Get(), TEXT("nvaftermathcallstack")) || (CVarAftermath_Callstack->GetInt())) ? GFSDK_Aftermath_FeatureFlags_CallStackCapturing     : 0;
 		Flags |= (FParse::Param(FCommandLine::Get(), TEXT("nvaftermathall"      )) || (CVarAftermath_TrackAll ->GetInt())) ? AllFlags                                            : 0;
 
+		const bool bDumpShaderDebugInfo = (FParse::Param(FCommandLine::Get(), TEXT("nvAftermathDumpShaderDebugInfo")) || (CVarAftermath_DumpShaderDebugInfo->GetInt()));
+
 		GFSDK_Aftermath_Result Result = GFSDK_Aftermath_EnableGpuCrashDumps(
 			GFSDK_Aftermath_Version_API,
 #if PLATFORM_WINDOWS
@@ -149,7 +160,7 @@ namespace UE::RHICore::Nvidia::Aftermath
 #endif
 			GFSDK_Aftermath_GpuCrashDumpFeatureFlags_Default,
 			&Callback_GpuCrashDump,
-			nullptr,
+			bDumpShaderDebugInfo ? &Callback_ShaderDebugInfo : nullptr,
 			&Callback_GpuCrashDumpDesc,
 			&Callback_ResolveMarker,
 			nullptr // user data
@@ -583,6 +594,25 @@ namespace UE::RHICore::Nvidia::Aftermath
 		}
 	}
 
+	void Callback_ShaderDebugInfo(const void* ShaderDebugInfo, const uint32 ShaderDebugInfoSize, void* UserData)
+	{
+		// Get shader debug information identifier.
+		GFSDK_Aftermath_ShaderDebugInfoIdentifier Identifier = {};
+		GFSDK_Aftermath_Result Result = GFSDK_Aftermath_GetShaderDebugInfoIdentifier(GFSDK_Aftermath_Version_API, ShaderDebugInfo, ShaderDebugInfoSize, &Identifier);
+
+		if (Result == GFSDK_Aftermath_Result_Success)
+		{
+			// Write to file for later in-depth analysis of crash dumps with Nsight Graphics.
+			const FString Filename = FPaths::ProjectLogDir() / *FString::Printf(TEXT("%016llX-%016llX.nvdbg"), Identifier.id[0], Identifier.id[1]);
+			FArchive* Writer = IFileManager::Get().CreateFileWriter(*Filename);
+			if (Writer)
+			{
+				Writer->Serialize((void*)ShaderDebugInfo, ShaderDebugInfoSize);
+				Writer->Close();
+			}
+		}
+	}
+
 	FCrashResult OnGPUCrash()
 	{
 		if (!bEnabled)
