Bindless Optimizations
- Adding a version of UE::RHICore::SetResourcesFromTables that skips the majority of resource iteration when all resources and samplers are bindless and RHI validation is off.
- The remaining resource iteration is to set SetLastRenderTime on all textures because that's still a thing.

diff --git a/Engine/Source/Runtime/RHICore/Public/RHICoreShader.h b/Engine/Source/Runtime/RHICore/Public/RHICoreShader.h
--- a/Engine/Source/Runtime/RHICore/Public/RHICoreShader.h
+++ b/Engine/Source/Runtime/RHICore/Public/RHICoreShader.h
@@ -112,7 +112,7 @@ template<> struct TResourceTypeStr<FRHIUnorderedAccessView> { static constexpr T
 template<> struct TResourceTypeStr<FRHIResourceCollection>  { static constexpr TCHAR String[] = TEXT("Resource Collection"); };
 
 template <typename TResourceType, typename TCallback>
-inline void EnumerateUniformBufferResources(FRHIUniformBuffer* RESTRICT Buffer, int32 BufferIndex, const uint32* RESTRICT ResourceMap, TCallback&& Callback)
+inline void EnumerateUniformBufferResources(const FRHIUniformBuffer* RESTRICT Buffer, int32 BufferIndex, const uint32* RESTRICT ResourceMap, TCallback&& Callback)
 {
 	const TRefCountPtr<FRHIResource>* RESTRICT Resources = Buffer->GetResourceTable().GetData();
 
@@ -144,7 +144,7 @@ inline void EnumerateUniformBufferResources(FRHIUniformBuffer* RESTRICT Buffer,
 	}
 }
 
-template <typename TBinder, typename TUniformBufferArrayType, typename TBitMaskType>
+template <typename TBinder, typename TUniformBufferArrayType, typename TBitMaskType, bool bFullyBindless = false>
 void SetResourcesFromTables(TBinder&& Binder, FRHIShader const& Shader, TBitMaskType& DirtyUniformBuffers, TUniformBufferArrayType const& BoundUniformBuffers
 #if ENABLE_RHI_VALIDATION
 	, RHIValidation::FTracker* Tracker
@@ -152,11 +152,34 @@ void SetResourcesFromTables(TBinder&& Binder, FRHIShader const& Shader, TBitMask
 )
 {
 	float CurrentTimeForTextureTimes = FApp::GetCurrentTime();
-
 	FShaderResourceTable const& SRT = Shader.GetShaderResourceTable();
 
 	// Mask the dirty bits by those buffers from which the shader has bound resources.
 	uint32 DirtyBits = SRT.ResourceTableBits & DirtyUniformBuffers;
+
+#if PLATFORM_SUPPORTS_BINDLESS_RENDERING && !ENABLE_RHI_VALIDATION
+	if constexpr (bFullyBindless)
+	{
+		while (DirtyBits)
+		{
+			// Scan for the lowest set bit, compute its index, clear it in the set of dirty bits.
+			const uint32 LowestBitMask = (DirtyBits) & (-(int32)DirtyBits);
+			const int32 BufferIndex = FMath::CountTrailingZeros(LowestBitMask); // todo: This has a branch on zero, we know it could never be zero...
+			DirtyBits ^= LowestBitMask;
+
+			EnumerateUniformBufferResources<FRHITexture>(BoundUniformBuffers[BufferIndex], BufferIndex, SRT.TextureMap.GetData(),
+				[&](FRHITexture* Texture, uint8 Index)
+				{
+					Texture->SetLastRenderTime(CurrentTimeForTextureTimes);
+				});
+		}
+
+		DirtyUniformBuffers = TBitMaskType(0);
+
+		return;
+	}
+#endif
+
 	while (DirtyBits)
 	{
 		// Scan for the lowest set bit, compute its index, clear it in the set of dirty bits.
@@ -296,5 +319,24 @@ void SetResourcesFromTables(TBinder&& Binder, FRHIShader const& Shader, TBitMask
 	DirtyUniformBuffers = TBitMaskType(0);
 }
 
+
+template <typename TBinder, typename TUniformBufferArrayType, typename TBitMaskType>
+void SetFullyBindlessResourcesFromTables(TBinder&& Binder, const FRHIShader& Shader, TBitMaskType& DirtyUniformBuffers, const TUniformBufferArrayType& BoundUniformBuffers
+#if ENABLE_RHI_VALIDATION
+	, RHIValidation::FTracker* Tracker
+#endif
+)
+{
+	SetResourcesFromTables<TBinder, TUniformBufferArrayType, TBitMaskType, true>(
+		MoveTemp(Binder)
+		, Shader
+		, DirtyUniformBuffers
+		, BoundUniformBuffers
+#if ENABLE_RHI_VALIDATION
+		, Tracker
+#endif
+	);
+}
+
 } //! RHICore
 } //! UE
