[RHICore]
* Fixed compile error when building clang modular build

diff --git a/Engine/Source/Runtime/RHICore/Private/RHICoreTransientResourceAllocator.cpp b/Engine/Source/Runtime/RHICore/Private/RHICoreTransientResourceAllocator.cpp
--- a/Engine/Source/Runtime/RHICore/Private/RHICoreTransientResourceAllocator.cpp
+++ b/Engine/Source/Runtime/RHICore/Private/RHICoreTransientResourceAllocator.cpp
@@ -85,7 +85,7 @@ DECLARE_DWORD_ACCUMULATOR_STAT(TEXT("Textures"), STAT_RHITransientTextures, STAT
 DECLARE_DWORD_ACCUMULATOR_STAT(TEXT("Buffers"), STAT_RHITransientBuffers, STATGROUP_RHITransientMemory);
 
 DECLARE_LLM_MEMORY_STAT(TEXT("RHI Transient Resources"), STAT_RHITransientResourcesLLM, STATGROUP_LLMFULL);
-RHICORE_API LLM_DEFINE_TAG(RHITransientResources, NAME_None, NAME_None, GET_STATFNAME(STAT_RHITransientResourcesLLM), GET_STATFNAME(STAT_EngineSummaryLLM));
+LLM_DEFINE_TAG(RHITransientResources, NAME_None, NAME_None, GET_STATFNAME(STAT_RHITransientResourcesLLM), GET_STATFNAME(STAT_EngineSummaryLLM));
 
 //////////////////////////////////////////////////////////////////////////
 
