Fixed static analysis build error.

diff --git a/Engine/Source/Runtime/RHICore/Private/RHICoreTransientResourceAllocator.cpp b/Engine/Source/Runtime/RHICore/Private/RHICoreTransientResourceAllocator.cpp
--- a/Engine/Source/Runtime/RHICore/Private/RHICoreTransientResourceAllocator.cpp
+++ b/Engine/Source/Runtime/RHICore/Private/RHICoreTransientResourceAllocator.cpp
@@ -688,7 +688,11 @@ FRHITransientBuffer* FRHITransientResourceHeapAllocator::CreateBufferInternal(
 {
 	FRHITransientBuffer* Buffer = nullptr;
 
+#if WITH_MGPU
 	ERHITransientHeapFlags BufferHeapFlag = (GNumExplicitGPUsForRendering > 1) && EnumHasAnyFlags(CreateInfo.Usage, EBufferUsageFlags::NNE) ? ERHITransientHeapFlags::AllowNNEBuffers : ERHITransientHeapFlags::AllowBuffers;
+#else
+	ERHITransientHeapFlags BufferHeapFlag = ERHITransientHeapFlags::AllowBuffers;
+#endif
 
 	for (FRHITransientHeap* Heap : Heaps)
 	{
