Add a crash context field to determine whether an Aftermath GPU crash dump file has been generated.

diff --git a/Engine/Source/Runtime/RHICore/Private/RHICoreNvidiaAftermath.cpp b/Engine/Source/Runtime/RHICore/Private/RHICoreNvidiaAftermath.cpp
--- a/Engine/Source/Runtime/RHICore/Private/RHICoreNvidiaAftermath.cpp
+++ b/Engine/Source/Runtime/RHICore/Private/RHICoreNvidiaAftermath.cpp
@@ -676,6 +676,7 @@ namespace UE::RHICore::Nvidia::Aftermath
 
 			case GFSDK_Aftermath_CrashDump_Status_Finished:
 				UE_LOG(LogNvidiaAftermath, Error, TEXT("%s"), *CrashResult.OutputLog);
+				FGenericCrashContext::SetEngineData(TEXT("RHI.HasAftermathCrashDump"), CrashResult.DumpPath.IsSet() ? TEXT("true") : TEXT("false"));
 				return CrashResult;
 			}
 		}
