[RHICore]
* Fixed compile error when building shipping with merged modules

diff --git a/Engine/Source/Runtime/RHICore/Private/RHICoreTransientResourceAllocator.cpp b/Engine/Source/Runtime/RHICore/Private/RHICoreTransientResourceAllocator.cpp
--- a/Engine/Source/Runtime/RHICore/Private/RHICoreTransientResourceAllocator.cpp
+++ b/Engine/Source/Runtime/RHICore/Private/RHICoreTransientResourceAllocator.cpp
@@ -85,7 +85,11 @@ DECLARE_DWORD_ACCUMULATOR_STAT(TEXT("Textures"), STAT_RHITransientTextures, STAT
 DECLARE_DWORD_ACCUMULATOR_STAT(TEXT("Buffers"), STAT_RHITransientBuffers, STATGROUP_RHITransientMemory);
 
 DECLARE_LLM_MEMORY_STAT(TEXT("RHI Transient Resources"), STAT_RHITransientResourcesLLM, STATGROUP_LLMFULL);
-RHICORE_API LLM_DEFINE_TAG(RHITransientResources, NAME_None, NAME_None, GET_STATFNAME(STAT_RHITransientResourcesLLM), GET_STATFNAME(STAT_EngineSummaryLLM));
+
+#if ENABLE_LOW_LEVEL_MEM_TRACKER // if LLM_DEFINE_TAG is something we need to dll export it
+RHICORE_API
+#endif
+LLM_DEFINE_TAG(RHITransientResources, NAME_None, NAME_None, GET_STATFNAME(STAT_RHITransientResourcesLLM), GET_STATFNAME(STAT_EngineSummaryLLM));
 
 //////////////////////////////////////////////////////////////////////////
 
