Disable NVIDIA Aftermath in DX11 due to driver instability.

Also move NVIDIA Aftermath initialization and DLL loading closer to RHI initialization. This ensures device profile CVars are applied and that the selected RHI is known.

diff --git a/Engine/Source/Runtime/RHICore/Private/RHICore.cpp b/Engine/Source/Runtime/RHICore/Private/RHICore.cpp
--- a/Engine/Source/Runtime/RHICore/Private/RHICore.cpp
+++ b/Engine/Source/Runtime/RHICore/Private/RHICore.cpp
@@ -5,17 +5,6 @@
 #include "RHICoreNvidiaAftermath.h"
 #include "Misc/CommandLine.h"
 
-class FRHICoreModule : public IModuleInterface
-{
-	virtual void StartupModule() override
-	{
-	#if NV_AFTERMATH
-		UE::RHICore::Nvidia::Aftermath::StartupModule();
-	#endif
-	}
-};
-
-IMPLEMENT_MODULE(FRHICoreModule, RHICore);
 DEFINE_LOG_CATEGORY(LogRHICore);
 
 namespace UE
diff --git a/Engine/Source/Runtime/RHICore/Private/RHICoreNvidiaAftermath.cpp b/Engine/Source/Runtime/RHICore/Private/RHICoreNvidiaAftermath.cpp
--- a/Engine/Source/Runtime/RHICore/Private/RHICoreNvidiaAftermath.cpp
+++ b/Engine/Source/Runtime/RHICore/Private/RHICoreNvidiaAftermath.cpp
@@ -100,9 +100,15 @@ namespace UE::RHICore::Nvidia::Aftermath
 	static void GFSDK_AFTERMATH_CALL Callback_ResolveMarker   (const void* MarkerData, const uint32_t MarkerDataSize, void* UserData, void** ResolvedMarkerData, uint32_t* ResolvedMarkerDataSize);
 	static void GFSDK_AFTERMATH_CALL Callback_ShaderDebugInfo (const void* ShaderDebugInfo, const uint32 ShaderDebugInfoSize, void* UserData);
 
-
-	void StartupModule()
+	// Loads the Aftermath DLL unless Aftermath initialization is disabled.
+	void LoadAftermathDLL()
 	{
+		if (GDynamicRHI && GDynamicRHI->GetInterfaceType() == ERHIInterfaceType::D3D11)
+		{
+			UE_LOG(LogNvidiaAftermath, Log, TEXT("Nvidia Aftermath is disabled in D3D11 due to instability issues."));
+			return;
+		}
+
 		if (!AllowVendorDevice())
 		{
 			UE_LOG(LogNvidiaAftermath, Log, TEXT("Vendor devices disallowed. Aftermath initialization skipped."));
@@ -129,6 +135,8 @@ namespace UE::RHICore::Nvidia::Aftermath
 
 	RHICORE_API void InitializeBeforeDeviceCreation(FResolveMarkerFunc ResolveMarkerFunc)
 	{
+		LoadAftermathDLL();
+
 		GResolveMarkerFunc = MoveTemp(ResolveMarkerFunc);
 
 		if (!DllHandle)
diff --git a/Engine/Source/Runtime/RHICore/Public/RHICoreNvidiaAftermath.h b/Engine/Source/Runtime/RHICore/Public/RHICoreNvidiaAftermath.h
--- a/Engine/Source/Runtime/RHICore/Public/RHICoreNvidiaAftermath.h
+++ b/Engine/Source/Runtime/RHICore/Public/RHICoreNvidiaAftermath.h
@@ -23,11 +23,6 @@
 		RHICORE_API bool IsEnabled();
 		RHICORE_API bool AreMarkersEnabled();
 
-		//
-		// Called on RHICore module startup to load the Nvidia Aftermath library
-		//
-		RHICORE_API void StartupModule();
-
 		//
 		// Called by platform RHIs to activate Aftermath.
 		// The optional ResolveMarkerFunc is provided for platform RHIs to replace the default marker resolution with something custom.
