Use aligned size as minimum pool size instead of requested size

diff --git a/Engine/Source/Runtime/RHICore/Private/RHIPoolAllocator.cpp b/Engine/Source/Runtime/RHICore/Private/RHIPoolAllocator.cpp
--- a/Engine/Source/Runtime/RHICore/Private/RHIPoolAllocator.cpp
+++ b/Engine/Source/Runtime/RHICore/Private/RHIPoolAllocator.cpp
@@ -702,14 +702,15 @@ bool FRHIPoolAllocator::TryAllocateInternal(uint32 InSizeInBytes, uint32 InAlloc
 		}
 	}
 
+	uint32 AlignedSize = FRHIMemoryPool::GetAlignedSize(InSizeInBytes, PoolAlignment, InAllocationAlignment);
 	if (NewPoolIndex >= 0)
 	{
-		Pools[NewPoolIndex] = CreateNewPool(NewPoolIndex, InSizeInBytes, InAllocationResourceType);
+		Pools[NewPoolIndex] = CreateNewPool(NewPoolIndex, AlignedSize, InAllocationResourceType);
 	}
 	else
 	{
 		NewPoolIndex = Pools.Num();		
-		Pools.Add(CreateNewPool(NewPoolIndex, InSizeInBytes, InAllocationResourceType));
+		Pools.Add(CreateNewPool(NewPoolIndex, AlignedSize, InAllocationResourceType));
 		PoolAllocationOrder.Add(NewPoolIndex);
 	}
 
