Added fix for crash when creating zero sized buffer in RDG. This is validated by RDG and will cause an ensure, but will no longer crash within the underlying allocator.

diff --git a/Engine/Source/Runtime/RHICore/Private/RHICoreTransientResourceAllocator.cpp b/Engine/Source/Runtime/RHICore/Private/RHICoreTransientResourceAllocator.cpp
--- a/Engine/Source/Runtime/RHICore/Private/RHICoreTransientResourceAllocator.cpp
+++ b/Engine/Source/Runtime/RHICore/Private/RHICoreTransientResourceAllocator.cpp
@@ -8,6 +8,7 @@
 #include "ProfilingDebugging/CsvProfiler.h"
 #include "RHICommandList.h"
 #include "RHICoreNvidiaAftermath.h"
+#include "RHICore.h"
 
 static bool GRHITransientAllocatorSeparateAsyncComputeHeap = false;
 static FAutoConsoleVariableRef CVarRHITransientAllocatorSeparateAsyncComputeHeap(
@@ -676,8 +677,14 @@ FRHITransientTexture* FRHITransientResourceHeapAllocator::CreateTextureInternal(
 
 		Texture = Heap->CreateTexture(CreateInfo, DebugName, Fences, CurrentCycle, TextureSize, TextureAlignment, CreateTextureFunction);
 	}
-	
-	check(Texture);
+
+	if (!Texture)
+	{
+		UE_LOG(LogRHICore, Fatal,
+			TEXT("Transient allocator failed to allocate Texture %s. Extent: (%u, %u), Depth: %u, ArraySize: %u, NumMips: %u. Allocation Size: %u, Allocation Alignment %u."),
+			DebugName, CreateInfo.Extent.X, CreateInfo.Extent.Y, CreateInfo.Depth, CreateInfo.ArraySize, CreateInfo.NumMips, TextureSize, TextureAlignment);
+	}
+
 	IF_RHICORE_TRANSIENT_ALLOCATOR_DEBUG(ActiveResources.Emplace(Texture));
 	return Texture;
 }
@@ -720,8 +727,14 @@ FRHITransientBuffer* FRHITransientResourceHeapAllocator::CreateBufferInternal(
 
 		Buffer = Heap->CreateBuffer(CreateInfo, DebugName, Fences, CurrentCycle, BufferSize, BufferAlignment, CreateBufferFunction);
 	}
-	
-	check(Buffer);
+
+	if (!Buffer)
+	{
+		UE_LOG(LogRHICore, Fatal,
+			TEXT("Transient allocator failed to allocate Buffer %s. Size: %u, Stride: %u. Allocation Size: %u, Allocation Alignment %u."),
+			DebugName, CreateInfo.Size, CreateInfo.Stride, BufferSize, BufferAlignment);
+	}
+
 	IF_RHICORE_TRANSIENT_ALLOCATOR_DEBUG(ActiveResources.Emplace(Buffer));
 	return Buffer;
 }
diff --git a/Engine/Source/Runtime/RenderCore/Public/RenderGraphBuilder.inl b/Engine/Source/Runtime/RenderCore/Public/RenderGraphBuilder.inl
--- a/Engine/Source/Runtime/RenderCore/Public/RenderGraphBuilder.inl
+++ b/Engine/Source/Runtime/RenderCore/Public/RenderGraphBuilder.inl
@@ -43,7 +43,6 @@ inline FRDGTextureRef FRDGBuilder::CreateTexture(
 	const TCHAR* Name,
 	ERDGTextureFlags Flags)
 {
-	// RDG no longer supports the legacy transient resource API.
 	FRDGTextureDesc OverrideDesc = Desc;
 
 #if !UE_BUILD_SHIPPING
@@ -68,10 +67,14 @@ inline FRDGBufferRef FRDGBuilder::CreateBuffer(
 	const TCHAR* Name,
 	ERDGBufferFlags Flags)
 {
-	// RDG no longer supports the legacy transient resource API.
 	FRDGBufferDesc OverrideDesc = Desc;
 
-	IF_RDG_ENABLE_DEBUG(UserValidation.ValidateCreateBuffer(Desc, Name, Flags));
+	IF_RDG_ENABLE_DEBUG(UserValidation.ValidateCreateBuffer(OverrideDesc, Name, Flags));
+
+	// Clamp the buffer size to that which is permissible, otherwise it's a guaranteed crash.
+	OverrideDesc.BytesPerElement = FMath::Clamp<uint32>(OverrideDesc.BytesPerElement, 1u, static_cast<uint32>(GetMaxBufferDimension()));
+	OverrideDesc.NumElements     = FMath::Clamp<uint32>(OverrideDesc.NumElements,     1u, static_cast<uint32>(GetMaxBufferDimension()));
+
 	FRDGBufferRef Buffer = Buffers.Allocate(Allocators.Root, Name, OverrideDesc, Flags);
 	IF_RDG_ENABLE_DEBUG(UserValidation.ValidateCreateBuffer(Buffer));
 	IF_RDG_ENABLE_TRACE(Trace.AddResource(Buffer));
