Disable Aftermath on Windows ASAN builds due to incompatibilities.

#rnx

diff --git a/Engine/Source/Runtime/RHICore/Private/RHICoreNvidiaAftermath.cpp b/Engine/Source/Runtime/RHICore/Private/RHICoreNvidiaAftermath.cpp
--- a/Engine/Source/Runtime/RHICore/Private/RHICoreNvidiaAftermath.cpp
+++ b/Engine/Source/Runtime/RHICore/Private/RHICoreNvidiaAftermath.cpp
@@ -28,6 +28,22 @@ THIRD_PARTY_INCLUDES_END
 
 DEFINE_LOG_CATEGORY_STATIC(LogNvidiaAftermath, Log, All);
 
+// Aftermath loads nvoglv64.dll on Windows. The enterprise/AWS version of this driver DLL hooks VirtualAlloc in its
+// DllMain function, which causes ASAN to abort. We're force-disabling Aftermath on ASAN Windows builds for this reason.
+#if PLATFORM_WINDOWS && PLATFORM_HAS_ASAN_INCLUDE
+#	if defined(__SANITIZE_ADDRESS__)
+#		define AFTERMATH_DISABLE_DUE_TO_ASAN 1
+#	elif defined(__has_feature)
+#		if __has_feature(address_sanitizer)
+#			define AFTERMATH_DISABLE_DUE_TO_ASAN 1
+#		endif
+#	endif
+#endif
+
+#ifndef AFTERMATH_DISABLE_DUE_TO_ASAN
+#	define AFTERMATH_DISABLE_DUE_TO_ASAN 0
+#endif
+
 namespace UE::RHICore::Nvidia::Aftermath
 {
 	static TAutoConsoleVariable<int32> CVarAftermath(
@@ -103,6 +119,9 @@ namespace UE::RHICore::Nvidia::Aftermath
 	// Loads the Aftermath DLL unless Aftermath initialization is disabled.
 	void LoadAftermathDLL()
 	{
+#if AFTERMATH_DISABLE_DUE_TO_ASAN
+		UE_LOG(LogNvidiaAftermath, Log, TEXT("Nvidia Aftermath is not compatible with Windows ASAN builds, so it is disabled."));
+#else
 		if (GDynamicRHI && GDynamicRHI->GetInterfaceType() == ERHIInterfaceType::D3D11)
 		{
 			UE_LOG(LogNvidiaAftermath, Log, TEXT("Nvidia Aftermath is disabled in D3D11 due to instability issues."));
@@ -131,6 +150,7 @@ namespace UE::RHICore::Nvidia::Aftermath
 		{
 			UE_LOG(LogNvidiaAftermath, Warning, TEXT("Failed to load GFSDK_Aftermath_Lib.x64.dll"));
 		}
+#endif
 	}
 
 	RHICORE_API void InitializeBeforeDeviceCreation(FResolveMarkerFunc ResolveMarkerFunc)
