Fixed Aftermath marker resolving interpreting string data as breadcrumbs

diff --git a/Engine/Source/Runtime/RHICore/Private/RHICoreNvidiaAftermath.cpp b/Engine/Source/Runtime/RHICore/Private/RHICoreNvidiaAftermath.cpp
--- a/Engine/Source/Runtime/RHICore/Private/RHICoreNvidiaAftermath.cpp
+++ b/Engine/Source/Runtime/RHICore/Private/RHICoreNvidiaAftermath.cpp
@@ -179,6 +179,8 @@ namespace UE::RHICore::Nvidia::Aftermath
 
 		const bool bDumpShaderDebugInfo = (FParse::Param(FCommandLine::Get(), TEXT("nvAftermathDumpShaderDebugInfo")) || (CVarAftermath_DumpShaderDebugInfo->GetInt()));
 
+		const bool bResolveMarkers = NV_AFTERMATH_USE_BREADCRUMB_PTRS || GResolveMarkerFunc.IsSet();
+		
 		GFSDK_Aftermath_Result Result = GFSDK_Aftermath_EnableGpuCrashDumps(
 			GFSDK_Aftermath_Version_API,
 #if PLATFORM_WINDOWS
@@ -190,7 +192,7 @@ namespace UE::RHICore::Nvidia::Aftermath
 			&Callback_GpuCrashDump,
 			bDumpShaderDebugInfo ? &Callback_ShaderDebugInfo : nullptr,
 			&Callback_GpuCrashDumpDesc,
-			&Callback_ResolveMarker,
+			bResolveMarkers ? &Callback_ResolveMarker : nullptr,
 			nullptr // user data
 		);
 
