Aftermath may report multiple crashes, potentially per device including internal ones, each must be tracked separately
- Fixed incorrect encoding on ILDN block debug name
- Fixed potential hang on pipeline cache consolidation

diff --git a/Engine/Source/Runtime/D3D12RHI/Private/D3D12NvidiaExtensions.cpp b/Engine/Source/Runtime/D3D12RHI/Private/D3D12NvidiaExtensions.cpp
--- a/Engine/Source/Runtime/D3D12RHI/Private/D3D12NvidiaExtensions.cpp
+++ b/Engine/Source/Runtime/D3D12RHI/Private/D3D12NvidiaExtensions.cpp
@@ -88,8 +88,9 @@ namespace UE::RHICore::Nvidia::Aftermath::D3D12
 
 		uint32_t IgnoredPipelines = 0;
 
+		// Get active pipelines, allow one second for consolidation to finish
 		TArray<TRefCountPtr<FRHIResource>> PipelineResources;
-		PipelineStateCache::GetPipelineStates(PipelineResources);
+		PipelineStateCache::GetPipelineStates(PipelineResources, true, FTimeout(FTimespan::FromSeconds(1)));
 
 		// Deduplicate shaders, Aftermath hashes are not local to the parent pipeline
 		for (FRHIResource* Resource : PipelineResources)
diff --git a/Engine/Source/Runtime/D3D12RHI/Private/D3D12Util.cpp b/Engine/Source/Runtime/D3D12RHI/Private/D3D12Util.cpp
--- a/Engine/Source/Runtime/D3D12RHI/Private/D3D12Util.cpp
+++ b/Engine/Source/Runtime/D3D12RHI/Private/D3D12Util.cpp
@@ -906,13 +906,18 @@ void FD3D12DynamicRHI::OutputGPUCrashReport(FTextBuilder& ErrorMessage)
 	}
 
 #if NV_AFTERMATH
-	UE::RHICore::Nvidia::Aftermath::FCrashResult AftermathResult = UE::RHICore::Nvidia::Aftermath::OnGPUCrash();
-	if (AftermathResult.GPUFaultAddress.IsSet())
+	TArray<UE::RHICore::Nvidia::Aftermath::FCrashResult> AftermathResults;
+	UE::RHICore::Nvidia::Aftermath::OnGPUCrash(AftermathResults);
+
+	for (const UE::RHICore::Nvidia::Aftermath::FCrashResult& AftermathResult : AftermathResults)
 	{
-		ForEachDevice(nullptr, [&](FD3D12Device* Device)
+		if (AftermathResult.GPUFaultAddress.IsSet())
 		{
-			D3D12RHI::LogPageFaultData(Device->GetParentAdapter(), Device, D3D12_GPU_VIRTUAL_ADDRESS(*AftermathResult.GPUFaultAddress));
-		});
+			ForEachDevice(nullptr, [&](FD3D12Device* Device)
+			{
+				D3D12RHI::LogPageFaultData(Device->GetParentAdapter(), Device, D3D12_GPU_VIRTUAL_ADDRESS(*AftermathResult.GPUFaultAddress));
+			});
+		}
 	}
 #endif
 #if INTEL_GPU_CRASH_DUMPS
@@ -945,11 +950,14 @@ void FD3D12DynamicRHI::OutputGPUCrashReport(FTextBuilder& ErrorMessage)
 	ErrorMessage.AppendLine(LOCTEXT("GPU Crashed", "GPU Crashed or D3D Device Removed.\n"));
 	ErrorMessage.AppendLine(LOCTEXT("GPU Crash Debugging enabled", "Check log for GPU state information."));
 #if NV_AFTERMATH
-	if (AftermathResult.DumpPath.IsSet())
+	for (const UE::RHICore::Nvidia::Aftermath::FCrashResult& AftermathResult : AftermathResults)
 	{
-		FFormatOrderedArguments Args;
-		Args.Add(FText::FromString(*AftermathResult.DumpPath));
-		ErrorMessage.AppendLineFormat(LOCTEXT("GPU CrashDump", "\nA GPU mini dump has been written to \"{0}\"."), Args);
+		if (AftermathResult.DumpPath.IsSet())
+		{
+			FFormatOrderedArguments Args;
+			Args.Add(FText::FromString(*AftermathResult.DumpPath));
+			ErrorMessage.AppendLineFormat(LOCTEXT("GPU CrashDump", "\nA GPU mini dump has been written to \"{0}\"."), Args);
+		}
 	}
 #endif
 }
diff --git a/Engine/Source/Runtime/RHI/Private/PipelineStateCache.cpp b/Engine/Source/Runtime/RHI/Private/PipelineStateCache.cpp
--- a/Engine/Source/Runtime/RHI/Private/PipelineStateCache.cpp
+++ b/Engine/Source/Runtime/RHI/Private/PipelineStateCache.cpp
@@ -1437,20 +1437,19 @@ public:
 		return true;
 	}
 
-	void GetResources(TArray<TRefCountPtr<FRHIResource>>& OutResources, bool bConsolidateWithInterrupt = true)
+	void GetResources(TArray<TRefCountPtr<FRHIResource>>& OutResources, bool bConsolidateWithInterrupt, UE::FTimeout ConsolidationTimeout)
 	{
 		FRWScopeLock InterruptGuard(InterruptLock, SLT_Write);
-		
-		if (bConsolidateWithInterrupt)
-		{
-			// Wait for any in-flight consolidation
-			WaitAndFinishAsyncCacheConsolidation();
 
+		// Wait for any in-flight consolidation.
+		// Consolidation is predicated on command context completion, allow for a timeout in case it's blocked
+		if (bConsolidateWithInterrupt && WaitAndFinishAsyncCacheConsolidation(ConsolidationTimeout))
+		{
 			bIsInterrupt = true;
 			
 			// Kick off a new one
 			FlushResources(false);
-			WaitAndFinishAsyncCacheConsolidation();
+			WaitAndFinishAsyncCacheConsolidation(ConsolidationTimeout);
 
 			bIsInterrupt = false;
 		}
@@ -1587,17 +1586,21 @@ public:
 	}
 
 private:
-	void WaitAndFinishAsyncCacheConsolidation()
+	bool WaitAndFinishAsyncCacheConsolidation(UE::FTimeout Timeout)
 	{
 		if (!RHICompletionEvent)
 		{
-			return;
+			return true;
 		}
-		
-		RHICompletionEvent->Wait();
+
+		if (!RHICompletionEvent->FTaskBase::Wait(Timeout))
+		{
+			return false;
+		}
+
 		RHICompletionEvent = nullptr;
-		
 		FinishAsyncCacheConsolidation();
+		return true;
 	}
 
 	void WaitTasksComplete(FPipelineStateCacheType* PipelineStateCache)
@@ -4436,11 +4439,11 @@ FRHIVertexDeclaration*	PipelineStateCache::GetOrCreateVertexDeclaration(const FV
 	return NewDeclaration;
 }
 
-void PipelineStateCache::GetPipelineStates(TArray<TRefCountPtr<FRHIResource>>& Out, bool bConsolidateCaches)
+void PipelineStateCache::GetPipelineStates(TArray<TRefCountPtr<FRHIResource>>& Out, bool bConsolidateCaches, UE::FTimeout ConsolidationTimeout)
 {
-	GComputePipelineCache.GetResources(Out, bConsolidateCaches);
-	GGraphicsPipelineCache.GetResources(Out, bConsolidateCaches);
-	GWorkGraphPipelineCache.GetResources(Out, bConsolidateCaches);
+	GComputePipelineCache.GetResources(Out, bConsolidateCaches, ConsolidationTimeout);
+	GGraphicsPipelineCache.GetResources(Out, bConsolidateCaches, ConsolidationTimeout);
+	GWorkGraphPipelineCache.GetResources(Out, bConsolidateCaches, ConsolidationTimeout);
 #if RHI_RAYTRACING
 	GRayTracingPipelineCache.GetResources(Out);
 #endif // RHI_RAYTRACING
diff --git a/Engine/Source/Runtime/RHI/Public/PipelineStateCache.h b/Engine/Source/Runtime/RHI/Public/PipelineStateCache.h
--- a/Engine/Source/Runtime/RHI/Public/PipelineStateCache.h
+++ b/Engine/Source/Runtime/RHI/Public/PipelineStateCache.h
@@ -8,6 +8,7 @@
 
 #include "CoreMinimal.h"
 #include "RHI.h"
+#include "Misc/Timeout.h"
 #include "Misc/EnumClassFlags.h"
 
 class FComputePipelineState;
@@ -130,7 +131,7 @@ namespace PipelineStateCache
 
 	extern RHI_API FGraphicsPipelineState*	FindGraphicsPipelineState(const FGraphicsPipelineStateInitializer& Initializer, bool bVerifyUse = true);
 
-	extern RHI_API void                     GetPipelineStates(TArray<TRefCountPtr<FRHIResource>>& Out, bool bConsolidateCaches = true);
+	extern RHI_API void                     GetPipelineStates(TArray<TRefCountPtr<FRHIResource>>& Out, bool bConsolidateCaches = true, UE::FTimeout ConsolidationTimeout = UE::FTimeout::Never());
 
 	extern RHI_API FRHIVertexDeclaration*	GetOrCreateVertexDeclaration(const FVertexDeclarationElementList& Elements);
 
diff --git a/Engine/Source/Runtime/RHICore/Internal/ShaderBinaryUtilities.h b/Engine/Source/Runtime/RHICore/Internal/ShaderBinaryUtilities.h
--- a/Engine/Source/Runtime/RHICore/Internal/ShaderBinaryUtilities.h
+++ b/Engine/Source/Runtime/RHICore/Internal/ShaderBinaryUtilities.h
@@ -89,7 +89,7 @@ namespace UE::ShaderBinaryUtilities
 		uint64 Offset = 0;
 	};
 	
-	static bool GetShaderBinaryDebugHashDXBC(const void* ShaderBinary, uint32 ByteSize, FUInt128& OutDebugHash)
+	static bool GetShaderBinaryDebugHashDXBC(const void* ShaderBinary, uint32 ByteSize, FString& OutDebugHash)
 	{
 		if (ByteSize < sizeof(DXBC::FHeader))
 		{
@@ -118,9 +118,10 @@ namespace UE::ShaderBinaryUtilities
 				{
 					auto DXILDebugInfo = ChunkCtx.Consume<DXIL::FShaderDebugNameInfo>();
 
-					constexpr uint32 HashLength = 32;
+					constexpr uint32 HashLength   = 32;
+					const     uint32 PdbExtLength = FCString::Strlen(TEXT(".pdb"));
 
-					if (DXILDebugInfo.NameLength != HashLength + FCString::Strlen(TEXT(".pdb")))
+					if (DXILDebugInfo.NameLength != HashLength + PdbExtLength)
 					{
 						UE_LOG(LogRHICore, Display, TEXT("DXIL name length not the expected hash"));
 						return false;
@@ -132,15 +133,18 @@ namespace UE::ShaderBinaryUtilities
 						return false;
 					}
 
-					String::HexToBytes(FUtf8StringView(static_cast<const UTF8CHAR*>(ChunkCtx.Data()), HashLength), reinterpret_cast<uint8*>(&OutDebugHash));
+					OutDebugHash = FString::ConstructFromPtrSize(static_cast<const ANSICHAR*>(ChunkCtx.Data()), HashLength);
 					break;
 				}
 				case 'HSAH':
 				{
 					auto DXILDebugInfo = ChunkCtx.Consume<DXIL::FShaderHashInfo>();
+					static_assert(sizeof(FUInt128) == sizeof(DXILDebugInfo.Digest), "Unexpected digest size");
+					
+					FUInt128 HashBytes;
+					FMemory::Memcpy(&HashBytes, DXILDebugInfo.Digest, sizeof(HashBytes));
 
-					static_assert(sizeof(OutDebugHash) == sizeof(DXILDebugInfo.Digest), "Unexpected digest size");
-					FMemory::Memcpy(&OutDebugHash, DXILDebugInfo.Digest, sizeof(OutDebugHash));
+					OutDebugHash = BytesToHex(reinterpret_cast<const uint8*>(&HashBytes), sizeof(FUInt128));
 					break;
 				}
 			}
diff --git a/Engine/Source/Runtime/RHICore/Private/RHICoreNvidiaAftermath.cpp b/Engine/Source/Runtime/RHICore/Private/RHICoreNvidiaAftermath.cpp
--- a/Engine/Source/Runtime/RHICore/Private/RHICoreNvidiaAftermath.cpp
+++ b/Engine/Source/Runtime/RHICore/Private/RHICoreNvidiaAftermath.cpp
@@ -123,13 +123,15 @@ namespace UE::RHICore::Nvidia::Aftermath
 	struct FShaderRegistration
 	{
 		FCriticalSection         Lock;
-		TMap<uint64, FUInt128>   ShaderBinaryToDebugHash;
+		TMap<uint64, FString>    ShaderBinaryToDebugHash;
 	} static ShaderRegistration;
 
 	struct FCrashIntermediates
 	{
 		FRWLock                     Lock;
+		TArray<FCrashResult>        Results;
 		TMap<uint64, TArray<uint8>> ShaderDebugLookup;
+		std::atomic_uint32_t        Counter{0};
 	} static CrashIntermediates;
 
 #if WITH_RHI_BREADCRUMBS
@@ -138,8 +140,6 @@ namespace UE::RHICore::Nvidia::Aftermath
 	static constexpr TCHAR const BreadcrumbsDisabledStr[] = TEXT("<RHI breadcrumbs disabled>");
 #endif
 
-	static FCrashResult CrashResult {};
-
 	RHICORE_API bool IsEnabled() { return bEnabled; }
 	RHICORE_API bool IsShaderRegistrationEnabled() { return bEnabled && CVarAftermath_ShaderRegistration->GetInt(); }
 	RHICORE_API bool AreMarkersEnabled() { return bEnabled && (Flags & GFSDK_Aftermath_FeatureFlags_EnableMarkers); }
@@ -293,7 +293,7 @@ namespace UE::RHICore::Nvidia::Aftermath
 		return true;
 	}
 
-	static void AppendGpuCrashShaderBinaryInfo(GFSDK_Aftermath_GpuCrashDump_Decoder Decoder, const GFSDK_Aftermath_GpuCrashDump_ShaderInfo& Shader)
+	static void AppendGpuCrashShaderBinaryInfo(GFSDK_Aftermath_GpuCrashDump_Decoder Decoder, const GFSDK_Aftermath_GpuCrashDump_ShaderInfo& Shader, FCrashResult& CrashResult)
 	{
 		GFSDK_Aftermath_ShaderBinaryHash BinaryHash;
 
@@ -307,7 +307,7 @@ namespace UE::RHICore::Nvidia::Aftermath
 		// Report binary hash separately, may not be the same
 		CrashResult.OutputLog += FString::Printf(TEXT("\n\t\t\tBinary Hash = %llu"), BinaryHash.hash);
 
-		FUInt128* Candidate = ShaderRegistration.ShaderBinaryToDebugHash.Find(BinaryHash.hash);
+		FString* Candidate = ShaderRegistration.ShaderBinaryToDebugHash.Find(BinaryHash.hash);
 		if (!Candidate)
 		{
 			CrashResult.OutputLog += TEXT("\n\t\t\t! Failed to associate Aftermath hash");
@@ -315,11 +315,10 @@ namespace UE::RHICore::Nvidia::Aftermath
 		}
 
 		// Report debug association
-		FString Hex = BytesToHex(reinterpret_cast<const uint8*>(&Candidate), sizeof(FUInt128));
-		CrashResult.OutputLog += FString::Printf(TEXT("\n\t\t\tPDB = %s.pdb"), *Hex);
+		CrashResult.OutputLog += FString::Printf(TEXT("\n\t\t\tPDB = %s.pdb"), **Candidate);
 	}
 
-	static void AppendGpuCrashShaderInfo(GFSDK_Aftermath_GpuCrashDump_Decoder Decoder)
+	static void AppendGpuCrashShaderInfo(GFSDK_Aftermath_GpuCrashDump_Decoder Decoder, FCrashResult& CrashResult)
 	{
 		uint32 ShaderCount = 0;
 
@@ -421,7 +420,7 @@ namespace UE::RHICore::Nvidia::Aftermath
 			CrashResult.OutputLog += FString::Printf(TEXT("\n\t\t\tType = %s"), ShaderType);
 			CrashResult.OutputLog += FString::Printf(TEXT("\n\t\t\tHash = %llu"), Shader.shaderHash);
 
-			AppendGpuCrashShaderBinaryInfo(Decoder, Shader);
+			AppendGpuCrashShaderBinaryInfo(Decoder, Shader, CrashResult);
 		}
 	}
 
@@ -448,9 +447,16 @@ namespace UE::RHICore::Nvidia::Aftermath
 			NameStorage.Reset();
 		#endif
 		};
+		
+		FCrashResult CrashResult;
 
 		// Write out crash dump to project log dir
-		CrashResult.DumpPath = FPaths::Combine(FPaths::ProjectLogDir(), FString::Printf(TEXT("%s.%s.nv-gpudmp"), GDynamicRHI->GetNonValidationRHI()->GetName(), *FDateTime::Now().ToString()));
+		CrashResult.DumpPath = FPaths::Combine(FPaths::ProjectLogDir(), FString::Printf(
+			TEXT("%s.%u.%s.nv-gpudmp"),
+			GDynamicRHI->GetNonValidationRHI()->GetName(),
+			CrashIntermediates.Counter++,
+			*FDateTime::Now().ToString()
+		));
 		FGenericCrashContext::AddFile(*CrashResult.DumpPath);
 		
 		CrashResult.OutputLog += FString::Printf(TEXT("\n\nWriting Aftermath dump to: %s"), *CrashResult.DumpPath.GetValue());
@@ -748,7 +754,7 @@ namespace UE::RHICore::Nvidia::Aftermath
 				GLateShaderAssociateFunc();
 			}
 			
-			AppendGpuCrashShaderInfo(Decoder);
+			AppendGpuCrashShaderInfo(Decoder, CrashResult);
 		}
 
 		// Append JSON dump generated by the decoder.
@@ -787,6 +793,13 @@ namespace UE::RHICore::Nvidia::Aftermath
 				CrashResult.OutputLog += FString::Printf(TEXT("\n\n%s\n\n"), StringCast<TCHAR>(JsonStr.GetData()).Get());
 			}
 		}
+
+		// Report the decoded results, this is to avoid cases where Aftermath fails to collect data
+		// post crash handling.
+		UE_LOG(LogNvidiaAftermath, Error, TEXT("%s"), *CrashResult.OutputLog);
+
+		FRWScopeLock Guard(CrashIntermediates.Lock, SLT_Write);
+		CrashIntermediates.Results.Emplace(MoveTemp(CrashResult));
 	}
 
 	static void GFSDK_AFTERMATH_CALL Callback_GpuCrashDumpDesc(PFN_GFSDK_Aftermath_AddGpuCrashDumpDescription AddValue, void* UserData)
@@ -870,7 +883,7 @@ namespace UE::RHICore::Nvidia::Aftermath
 		}
 
 		// On disk debug hash, typically the [hash].pdb
-		FUInt128 DebugHash;
+		FString DebugHash;
 
 		// Internal aftermath hash
 		GFSDK_Aftermath_ShaderBinaryHash AftermathHash{};
@@ -924,9 +937,18 @@ namespace UE::RHICore::Nvidia::Aftermath
 	}
 
 	FCrashResult OnGPUCrash()
+	{
+		TArray<FCrashResult> Results;
+		OnGPUCrash(Results);
+
+		// Just report the last one, this method is pending deprecation
+		return !Results.IsEmpty() ? Results.Last() : FCrashResult {};
+	}
+	
+	bool OnGPUCrash(TArray<FCrashResult>& OutResults)
 	{
 		if (!bEnabled)
-			return {};
+			return false;
 
 		const double StartTime = FPlatformTime::Seconds();
 		const double ProcessingEndTime = StartTime + CVar_DumpProcessWaitTime->GetFloat();
@@ -941,49 +963,67 @@ namespace UE::RHICore::Nvidia::Aftermath
 			if (Result != GFSDK_Aftermath_Result_Success)
 			{
 				UE_LOG(LogNvidiaAftermath, Warning, TEXT("GFSDK_Aftermath_GetCrashDumpStatus failed: 0x%08x. Skipping crash dump processing."), Result);
-				return {};
+				return false;
 			}
 
 			switch (Status)
 			{
-			default: [[fallthrough]];
-			case GFSDK_Aftermath_CrashDump_Status_Unknown:
-				UE_LOG(LogNvidiaAftermath, Warning, TEXT("Unknown aftermath crash dump state. Skipping crash dump processing."));
-				return {};
-			case GFSDK_Aftermath_CrashDump_Status_NotStarted:
-				// Driver may take time before starting the dump
-				if (FPlatformTime::Seconds() >= NotStartedEndTime)
+				default: [[fallthrough]];
+				case GFSDK_Aftermath_CrashDump_Status_Unknown:
 				{
-					UE_LOG(LogNvidiaAftermath, Warning, TEXT("Timed out while waiting for Aftermath to start the GPU crash dump."));
-					return {};
+					UE_LOG(LogNvidiaAftermath, Warning, TEXT("Unknown aftermath crash dump state. Skipping crash dump processing."));
+					return false;
 				}
-				else
+				case GFSDK_Aftermath_CrashDump_Status_NotStarted:
 				{
-					FPlatformProcess::Sleep(0.01f);
+					// Driver may take time before starting the dump
+					if (FPlatformTime::Seconds() >= NotStartedEndTime)
+					{
+						UE_LOG(LogNvidiaAftermath, Warning, TEXT("Timed out while waiting for Aftermath to start the GPU crash dump."));
+						return false;
+					}
+					else
+					{
+						FPlatformProcess::Sleep(0.01f);
+					}
+					break;
 				}
-				break;
-			case GFSDK_Aftermath_CrashDump_Status_CollectingData:
-			case GFSDK_Aftermath_CrashDump_Status_InvokingCallback:
-				// Crash dump is still in progress...
-				if (FPlatformTime::Seconds() >= ProcessingEndTime)
+				case GFSDK_Aftermath_CrashDump_Status_CollectingData:
+				case GFSDK_Aftermath_CrashDump_Status_InvokingCallback:
 				{
-					UE_LOG(LogNvidiaAftermath, Warning, TEXT("Timed out while waiting for Aftermath to complete GPU crash dump."));
-					return {};
+					// Crash dump is still in progress...
+					if (FPlatformTime::Seconds() >= ProcessingEndTime)
+					{
+						UE_LOG(LogNvidiaAftermath, Warning, TEXT("Timed out while waiting for Aftermath to complete GPU crash dump."));
+						return false;
+					}
+					else
+					{
+						FPlatformProcess::Sleep(0.01f);
+					}
+					break;
 				}
-				else
+
+				case GFSDK_Aftermath_CrashDump_Status_CollectingDataFailed:
 				{
-					FPlatformProcess::Sleep(0.01f);
+					UE_LOG(LogNvidiaAftermath, Warning, TEXT("Aftermath failed to collect GPU crash data."));
+					return false;
 				}
-				break;
 
-			case GFSDK_Aftermath_CrashDump_Status_CollectingDataFailed:
-				UE_LOG(LogNvidiaAftermath, Warning, TEXT("Aftermath failed to collect GPU crash data."));
-				return {};
+				case GFSDK_Aftermath_CrashDump_Status_Finished:
+				{
+					FRWScopeLock Guard(CrashIntermediates.Lock, SLT_ReadOnly);
+					OutResults = MoveTemp(CrashIntermediates.Results);
 
-			case GFSDK_Aftermath_CrashDump_Status_Finished:
-				UE_LOG(LogNvidiaAftermath, Error, TEXT("%s"), *CrashResult.OutputLog);
-				FGenericCrashContext::SetEngineData(TEXT("RHI.HasAftermathCrashDump"), CrashResult.DumpPath.IsSet() ? TEXT("true") : TEXT("false"));
-				return CrashResult;
+					// Any valid dump?
+					const bool bHasDump = OutResults.ContainsByPredicate([](const FCrashResult& Result)
+					{
+						return Result.DumpPath.IsSet();
+					});
+
+					FGenericCrashContext::SetEngineData(TEXT("RHI.HasAftermathCrashDump"), bHasDump ? TEXT("true") : TEXT("false"));
+					return true;
+				}
 			}
 		}
 	}
diff --git a/Engine/Source/Runtime/RHICore/Public/RHICoreNvidiaAftermath.h b/Engine/Source/Runtime/RHICore/Public/RHICoreNvidiaAftermath.h
--- a/Engine/Source/Runtime/RHICore/Public/RHICoreNvidiaAftermath.h
+++ b/Engine/Source/Runtime/RHICore/Public/RHICoreNvidiaAftermath.h
@@ -67,10 +67,17 @@
 
 		//
 		// Called by platform RHIs when a GPU crash is detected.
-		// Waits for Aftermath to finish crash dump processing, then results the result.
+		// Waits for Aftermath to finish crash dump processing, then returns the result.
 		//
+		UE_DEPRECATED(5.6, "Use OnGPUCrash(OutResults) overload instead, Aftermath may report multiple crashes")
 		RHICORE_API FCrashResult OnGPUCrash();
 
+		//
+		// Called by platform RHIs when a GPU crash is detected.
+		// Waits for Aftermath to finish crash dump processing, then returns the result.
+		//
+		RHICORE_API bool OnGPUCrash(TArray<FCrashResult>& OutResults);
+
 		static inline constexpr TCHAR RootNodeName[] = TEXT("<root>");
 
 		//
diff --git a/Engine/Source/Runtime/VulkanRHI/Private/VulkanUtil.cpp b/Engine/Source/Runtime/VulkanRHI/Private/VulkanUtil.cpp
--- a/Engine/Source/Runtime/VulkanRHI/Private/VulkanUtil.cpp
+++ b/Engine/Source/Runtime/VulkanRHI/Private/VulkanUtil.cpp
@@ -930,7 +930,8 @@ namespace VulkanRHI
 
 			// Make sure we wait on the Aftermath crash dump before we crash.
 #if NV_AFTERMATH
-			UE::RHICore::Nvidia::Aftermath::OnGPUCrash();
+			TArray<UE::RHICore::Nvidia::Aftermath::FCrashResult> AftermathResults;
+			UE::RHICore::Nvidia::Aftermath::OnGPUCrash(AftermathResults);
 #endif
 
 			UGameEngine* GameEngine = Cast<UGameEngine>(GEngine);
diff --git a/Engine/Source/Runtime/Windows/D3D11RHI/Private/D3D11Util.cpp b/Engine/Source/Runtime/Windows/D3D11RHI/Private/D3D11Util.cpp
--- a/Engine/Source/Runtime/Windows/D3D11RHI/Private/D3D11Util.cpp
+++ b/Engine/Source/Runtime/Windows/D3D11RHI/Private/D3D11Util.cpp
@@ -96,7 +96,8 @@ static void TerminateOnDeviceRemoved(HRESULT D3DResult, ID3D11Device* Direct3DDe
 	if (D3DResult == DXGI_ERROR_DEVICE_REMOVED)
 	{
 	#if NV_AFTERMATH
-		UE::RHICore::Nvidia::Aftermath::OnGPUCrash();
+		TArray<UE::RHICore::Nvidia::Aftermath::FCrashResult> AftermathResults;
+		UE::RHICore::Nvidia::Aftermath::OnGPUCrash(AftermathResults);
 	#endif
 
 		UGameEngine* GameEngine = Cast<UGameEngine>(GEngine);
