Added Vulkan Aftermath late shader association support
- Added debug names during registration, future association will likely also report DebugStoragePath
- Added pipeline frame counters, avoid late association on pipelines that havent been used in N frames

diff --git a/Engine/Source/Runtime/D3D12RHI/Private/D3D12NvidiaExtensions.cpp b/Engine/Source/Runtime/D3D12RHI/Private/D3D12NvidiaExtensions.cpp
--- a/Engine/Source/Runtime/D3D12RHI/Private/D3D12NvidiaExtensions.cpp
+++ b/Engine/Source/Runtime/D3D12RHI/Private/D3D12NvidiaExtensions.cpp
@@ -24,20 +24,6 @@ THIRD_PARTY_INCLUDES_END
 	#include "Windows/HideWindowsPlatformTypes.h"
 #endif
 
-static TAutoConsoleVariable<float> CVarAftermathLateShaderAssociationsTimeLimit(
-	TEXT("r.GPUCrashDebugging.Aftermath.D3D12.LateShaderAssociations.TimeLimit"),
-	7.5f,
-	TEXT("Time limit (s) before the late associations are stopped\n"),
-	ECVF_Default
-);
-
-static TAutoConsoleVariable<int32> CVarAftermathLateShaderAssociationFrameLimit(
-	TEXT("r.GPUCrashDebugging.Aftermath.D3D12.LateShaderAssociations.FrameLimit"),
-	10,
-	TEXT("Max last since used number of frames to consider when collecting pipelines\n"),
-	ECVF_Default
-);
-
 namespace UE::RHICore::Nvidia::Aftermath::D3D12
 {
 	void InitializeDevice(ID3D12Device* RootDevice)
@@ -73,14 +59,10 @@ namespace UE::RHICore::Nvidia::Aftermath::D3D12
 		return (PipelineFrameIndex <= AdapterFrameIndex) && (AdapterFrameIndex - PipelineFrameIndex) > Threshold;
 	}
 
-	void CreateShaderAssociations()
+	void CreateShaderAssociations(float TimeLimitSeconds, uint32 FrameLimit)
 	{
-		float TimeLimit = CVarAftermathLateShaderAssociationsTimeLimit->GetFloat();
-		
 		uint64 CycleStart = FPlatformTime::Cycles64();
 
-		uint32 FrameThreshold = CVarAftermathLateShaderAssociationFrameLimit->GetInt();
-
 		UE_LOG(LogD3D12RHI, Log, TEXT("Starting late shader associations..."));
 
 		TArray<FRHIShader*> Shaders;
@@ -95,6 +77,11 @@ namespace UE::RHICore::Nvidia::Aftermath::D3D12
 		// Deduplicate shaders, Aftermath hashes are not local to the parent pipeline
 		for (FRHIResource* Resource : PipelineResources)
 		{
+			if (!Resource)
+			{
+				continue;
+			}
+			
 			switch (Resource->GetType())
 			{
 				default:
@@ -103,7 +90,7 @@ namespace UE::RHICore::Nvidia::Aftermath::D3D12
 				{
 					auto* Pipeline = static_cast<FD3D12GraphicsPipelineState*>(Resource);
 						
-					if (IsPipelineUnused(Pipeline->PipelineState->GetParentAdapter(), Pipeline->FrameCounter, FrameThreshold))
+					if (IsPipelineUnused(Pipeline->PipelineState->GetParentAdapter(), Pipeline->FrameCounter, FrameLimit))
 					{
 						IgnoredPipelines++;
 						continue;
@@ -120,7 +107,7 @@ namespace UE::RHICore::Nvidia::Aftermath::D3D12
 				{
 					auto* Pipeline = static_cast<FD3D12ComputePipelineState*>(Resource);
 					
-					if (IsPipelineUnused(Pipeline->PipelineState->GetParentAdapter(), Pipeline->FrameCounter, FrameThreshold))
+					if (IsPipelineUnused(Pipeline->PipelineState->GetParentAdapter(), Pipeline->FrameCounter, FrameLimit))
 					{
 						IgnoredPipelines++;
 						continue;
@@ -134,7 +121,7 @@ namespace UE::RHICore::Nvidia::Aftermath::D3D12
 				{
 					auto* Pipeline = static_cast<FD3D12RayTracingPipelineState*>(Resource);
 						
-					if (IsPipelineUnused(Pipeline->Device->GetParentAdapter(), Pipeline->FrameCounter, FrameThreshold))
+					if (IsPipelineUnused(Pipeline->Device->GetParentAdapter(), Pipeline->FrameCounter, FrameLimit))
 					{
 						IgnoredPipelines++;
 						continue;
@@ -164,7 +151,7 @@ namespace UE::RHICore::Nvidia::Aftermath::D3D12
 				{
 					auto* Pipeline = static_cast<FD3D12WorkGraphPipelineState*>(Resource);
 						
-					if (IsPipelineUnused(Pipeline->Device->GetParentAdapter(), Pipeline->FrameCounter, FrameThreshold))
+					if (IsPipelineUnused(Pipeline->Device->GetParentAdapter(), Pipeline->FrameCounter, FrameLimit))
 					{
 						IgnoredPipelines++;
 						continue;
@@ -184,11 +171,11 @@ namespace UE::RHICore::Nvidia::Aftermath::D3D12
 		UE_LOG(LogD3D12RHI, Log, TEXT("Late shader associations ignored %u pipelines based on frame fences"), IgnoredPipelines);
 
 		// Parallelize as much as possible to avoid timeouts
-		ParallelFor(Shaders.Num(), [CycleStart, TimeLimit, &Shaders](int32 Index)
+		ParallelFor(Shaders.Num(), [CycleStart, TimeLimitSeconds, &Shaders](int32 Index)
 		{
 			// Aftermath handling is time constrained, if we hit the limit just stop
 			float Elapsed = FPlatformTime::ToSeconds64(FPlatformTime::Cycles64() - CycleStart);
-			if (Elapsed >= TimeLimit)
+			if (Elapsed >= TimeLimitSeconds)
 			{
 				UE_CALL_ONCE([Elapsed]
 				{
diff --git a/Engine/Source/Runtime/D3D12RHI/Private/D3D12NvidiaExtensions.h b/Engine/Source/Runtime/D3D12RHI/Private/D3D12NvidiaExtensions.h
--- a/Engine/Source/Runtime/D3D12RHI/Private/D3D12NvidiaExtensions.h
+++ b/Engine/Source/Runtime/D3D12RHI/Private/D3D12NvidiaExtensions.h
@@ -25,7 +25,7 @@
 		typedef struct GFSDK_Aftermath_ResourceHandle__* FResource;
 
 		void InitializeDevice(ID3D12Device* RootDevice);
-		void CreateShaderAssociations();
+		void CreateShaderAssociations(float TimeLimitSeconds, uint32 FrameLimit);
 
 		FCommandList RegisterCommandList(ID3D12CommandList* D3DCommandList);
 		void UnregisterCommandList(FCommandList CommandList);
diff --git a/Engine/Source/Runtime/RHICore/Internal/ShaderBinaryUtilities.h b/Engine/Source/Runtime/RHICore/Internal/ShaderBinaryUtilities.h
--- a/Engine/Source/Runtime/RHICore/Internal/ShaderBinaryUtilities.h
+++ b/Engine/Source/Runtime/RHICore/Internal/ShaderBinaryUtilities.h
@@ -118,10 +118,9 @@ namespace UE::ShaderBinaryUtilities
 				{
 					auto DXILDebugInfo = ChunkCtx.Consume<DXIL::FShaderDebugNameInfo>();
 
-					constexpr uint32 HashLength   = 32;
-					const     uint32 PdbExtLength = FCString::Strlen(TEXT(".pdb"));
+					const uint32 HashLength = 32 + FCString::Strlen(TEXT(".pdb"));
 
-					if (DXILDebugInfo.NameLength != HashLength + PdbExtLength)
+					if (DXILDebugInfo.NameLength != HashLength)
 					{
 						UE_LOG(LogRHICore, Display, TEXT("DXIL name length not the expected hash"));
 						return false;
diff --git a/Engine/Source/Runtime/RHICore/Private/RHICoreNvidiaAftermath.cpp b/Engine/Source/Runtime/RHICore/Private/RHICoreNvidiaAftermath.cpp
--- a/Engine/Source/Runtime/RHICore/Private/RHICoreNvidiaAftermath.cpp
+++ b/Engine/Source/Runtime/RHICore/Private/RHICoreNvidiaAftermath.cpp
@@ -19,6 +19,7 @@
 #endif
 
 THIRD_PARTY_INCLUDES_START
+	#define VULKAN_H_
 	#if PLATFORM_WINDOWS
 		#include "d3d12.h"
 	#endif // PLATFORM_WINDOWS
@@ -100,14 +101,28 @@ namespace UE::RHICore::Nvidia::Aftermath
 		ECVF_ReadOnly
 	);
 
-	static TAutoConsoleVariable<float> CVar_DumpProcessWaitTime(
+	static TAutoConsoleVariable<float> CVarAftermath_LateShaderAssociationsTimeLimit(
+		TEXT("r.GPUCrashDebugging.Aftermath.LateShaderAssociations.TimeLimit"),
+		7.5f,
+		TEXT("Time limit (s) before the late associations are stopped\n"),
+		ECVF_Default
+	);
+
+	static TAutoConsoleVariable<int32> CVarAftermath_LateShaderAssociationsFrameLimit(
+		TEXT("r.GPUCrashDebugging.Aftermath.LateShaderAssociations.FrameLimit"),
+		10,
+		TEXT("Max last since used number of frames to consider when collecting pipelines\n"),
+		ECVF_Default
+	);
+
+	static TAutoConsoleVariable<float> CVarAftermath_DumpProcessWaitTime(
 		TEXT("r.GPUCrashDebugging.Aftermath.DumpProcessWaitTime"),
 		10.0f,
 		TEXT("Amount of time (in seconds) to wait for Aftermath to finish processing GPU crash dumps."),
 		ECVF_Default
 	);
 
-	static TAutoConsoleVariable<float> CVar_DumpStartWaitTime(
+	static TAutoConsoleVariable<float> CVarAftermath_DumpStartWaitTime(
 		TEXT("r.GPUCrashDebugging.Aftermath.DumpStartWaitTime"),
 		3.0f,
 		TEXT("Amount of time (in seconds) to wait for Aftermath to start processing GPU crash dumps."),
@@ -119,11 +134,17 @@ namespace UE::RHICore::Nvidia::Aftermath
 	static uint32 Flags = 0;
 	static FResolveMarkerFunc GResolveMarkerFunc;
 	static FAssociateShadersFunc GLateShaderAssociateFunc;
+
+	struct FShaderDebugInfo
+	{
+		FString HashName;
+		FString DebugName;
+	};
 	
 	struct FShaderRegistration
 	{
-		FCriticalSection         Lock;
-		TMap<uint64, FString>    ShaderBinaryToDebugHash;
+		FCriticalSection               Lock;
+		TMap<uint64, FShaderDebugInfo> ShaderBinaryToDebugHash;
 	} static ShaderRegistration;
 
 	struct FCrashIntermediates
@@ -307,15 +328,22 @@ namespace UE::RHICore::Nvidia::Aftermath
 		// Report binary hash separately, may not be the same
 		CrashResult.OutputLog += FString::Printf(TEXT("\n\t\t\tBinary Hash = %llu"), BinaryHash.hash);
 
-		FString* Candidate = ShaderRegistration.ShaderBinaryToDebugHash.Find(BinaryHash.hash);
+		FShaderDebugInfo* Candidate = ShaderRegistration.ShaderBinaryToDebugHash.Find(BinaryHash.hash);
 		if (!Candidate)
 		{
 			CrashResult.OutputLog += TEXT("\n\t\t\t! Failed to associate Aftermath hash");
 			return;
 		}
 
-		// Report debug association
-		CrashResult.OutputLog += FString::Printf(TEXT("\n\t\t\tPDB = %s.pdb"), **Candidate);
+		if (!Candidate->HashName.IsEmpty())
+		{
+			CrashResult.OutputLog += FString::Printf(TEXT("\n\t\t\tHash = %s"), *Candidate->HashName);
+		}
+
+		if (!Candidate->DebugName.IsEmpty())
+		{
+			CrashResult.OutputLog += FString::Printf(TEXT("\n\t\t\tDebugName = '%s'"), *Candidate->DebugName);
+		}
 	}
 
 	static void AppendGpuCrashShaderInfo(GFSDK_Aftermath_GpuCrashDump_Decoder Decoder, FCrashResult& CrashResult)
@@ -751,7 +779,10 @@ namespace UE::RHICore::Nvidia::Aftermath
 			// Allow association after a fault
 			if (GLateShaderAssociateFunc)
 			{
-				GLateShaderAssociateFunc();
+				GLateShaderAssociateFunc(
+					CVarAftermath_LateShaderAssociationsTimeLimit->GetFloat(),
+					CVarAftermath_LateShaderAssociationsFrameLimit->GetInt()
+				);
 			}
 			
 			AppendGpuCrashShaderInfo(Decoder, CrashResult);
@@ -874,7 +905,7 @@ namespace UE::RHICore::Nvidia::Aftermath
 		}
 	}
 
-	FShaderHash RegisterShaderBinary(const void* ShaderBinary, uint32 ByteSize)
+	FShaderHash RegisterShaderBinary(const void* ShaderBinary, uint32 ByteSize, const FStringView& DebugName)
 	{
 		if (ByteSize < sizeof(uint32))
 		{
@@ -882,8 +913,12 @@ namespace UE::RHICore::Nvidia::Aftermath
 			return InvalidShaderHash;
 		}
 
-		// On disk debug hash, typically the [hash].pdb
-		FString DebugHash;
+		FShaderDebugInfo DebugInfo;
+
+		if (!DebugName.IsEmpty())
+		{
+			DebugInfo.DebugName = DebugName;
+		}
 
 		// Internal aftermath hash
 		GFSDK_Aftermath_ShaderBinaryHash AftermathHash{};
@@ -897,7 +932,7 @@ namespace UE::RHICore::Nvidia::Aftermath
 			}
 			case 'CBXD': // DXBC container header
 			{
-				ShaderBinaryUtilities::GetShaderBinaryDebugHashDXBC(ShaderBinary, ByteSize, DebugHash);
+				ShaderBinaryUtilities::GetShaderBinaryDebugHashDXBC(ShaderBinary, ByteSize, DebugInfo.HashName);
 					
 #ifdef __d3d12_h__
 				D3D12_SHADER_BYTECODE ByteCode;
@@ -911,7 +946,7 @@ namespace UE::RHICore::Nvidia::Aftermath
 					return InvalidShaderHash;
 				}
 #else // __d3d12_h__
-				UE_LOG(LogNvidiaAftermath, Display, TEXT("DXBC container without D3D12 support"), Result);
+				UE_LOG(LogNvidiaAftermath, Display, TEXT("DXBC container without D3D12 support"));
 				return InvalidShaderHash;
 #endif // __d3d12_h__
 					
@@ -919,13 +954,28 @@ namespace UE::RHICore::Nvidia::Aftermath
 			}
 			case 0x07230203: // Vulkan magic header
 			{
-				UE_LOG(LogRHICore, Display, TEXT("SPIRV binary debug hash not supported yet"));
+				// Consider extracting DebugStoragePath from NonSemantic.Shader.DebugInfo.100
+					
+#ifdef VULKAN_H_
+				GFSDK_Aftermath_SpirvCode SPIRV;
+				SPIRV.pData = ShaderBinary;
+				SPIRV.size = ByteSize;
+					
+				GFSDK_Aftermath_Result Result = GFSDK_Aftermath_GetShaderHashSpirv(GFSDK_Aftermath_Version_API, &SPIRV, &AftermathHash);
+				if (Result != GFSDK_Aftermath_Result_Success)
+				{
+					UE_LOG(LogNvidiaAftermath, Display, TEXT("Failed to get aftermath shader hash (%u)"), Result);
+					return InvalidShaderHash;
+				}
+#else // VULKAN_H_
+				UE_LOG(LogNvidiaAftermath, Display, TEXT("SPIRV container without Vulkan support"));
 				return InvalidShaderHash;
+#endif // VULKAN_H_
 			}
 		}
 
 		FScopeLock Lock(&ShaderRegistration.Lock);
-		ShaderRegistration.ShaderBinaryToDebugHash.Add(AftermathHash.hash, DebugHash);
+		ShaderRegistration.ShaderBinaryToDebugHash.Add(AftermathHash.hash, DebugInfo);
 		
 		return AftermathHash.hash;
 	}
@@ -951,8 +1001,8 @@ namespace UE::RHICore::Nvidia::Aftermath
 			return false;
 
 		const double StartTime = FPlatformTime::Seconds();
-		const double ProcessingEndTime = StartTime + CVar_DumpProcessWaitTime->GetFloat();
-		const double NotStartedEndTime = StartTime + CVar_DumpStartWaitTime->GetFloat();
+		const double ProcessingEndTime = StartTime + CVarAftermath_DumpProcessWaitTime->GetFloat();
+		const double NotStartedEndTime = StartTime + CVarAftermath_DumpStartWaitTime->GetFloat();
 
 		// Wait for GPU crash dump processing to complete.
 		while (true)
diff --git a/Engine/Source/Runtime/RHICore/Public/RHICoreNvidiaAftermath.h b/Engine/Source/Runtime/RHICore/Public/RHICoreNvidiaAftermath.h
--- a/Engine/Source/Runtime/RHICore/Public/RHICoreNvidiaAftermath.h
+++ b/Engine/Source/Runtime/RHICore/Public/RHICoreNvidiaAftermath.h
@@ -19,7 +19,7 @@
 	namespace UE::RHICore::Nvidia::Aftermath
 	{
 		typedef TUniqueFunction<void(const void* MarkerData, const uint32_t MarkerDataSize, void* UserData, void** ResolvedMarkerData, uint32_t* ResolvedMarkerDataSize)> FResolveMarkerFunc;
-		typedef TUniqueFunction<void()> FAssociateShadersFunc;
+		typedef TUniqueFunction<void(float TimeLimitSeconds, uint32 FrameLimit)> FAssociateShadersFunc;
 		
 		RHICORE_API bool IsEnabled();
 		RHICORE_API bool IsShaderRegistrationEnabled();
@@ -50,7 +50,7 @@
 		// Register a new shader binary with Aftermath, must be valid.
 		// Returns the internal shader hash, must be deregistered with DeregisterShaderBinary
 		//
-		RHICORE_API FShaderHash RegisterShaderBinary(const void* ShaderBinary, uint32 ByteSize);
+		RHICORE_API FShaderHash RegisterShaderBinary(const void* ShaderBinary, uint32 ByteSize, const FStringView& DebugName = nullptr);
 
 		//
 		// Deregister a shader binary from its hash. Will not release it immediately,
diff --git a/Engine/Source/Runtime/VulkanRHI/Private/VulkanExtensions.cpp b/Engine/Source/Runtime/VulkanRHI/Private/VulkanExtensions.cpp
--- a/Engine/Source/Runtime/VulkanRHI/Private/VulkanExtensions.cpp
+++ b/Engine/Source/Runtime/VulkanRHI/Private/VulkanExtensions.cpp
@@ -1198,6 +1198,9 @@ public:
 			{
 				AftermathResolveMarkerCallback(MarkerData, ResolvedMarkerData, ResolvedMarkerDataSize);
 			});
+
+			extern void AftermathLateAssociate(float TimeLimitSeconds, uint32 FrameLimit);
+			UE::RHICore::Nvidia::Aftermath::SetLateShaderAssociateCallback(AftermathLateAssociate);
 #endif
 		}
 	}
diff --git a/Engine/Source/Runtime/VulkanRHI/Private/VulkanPipeline.h b/Engine/Source/Runtime/VulkanRHI/Private/VulkanPipeline.h
--- a/Engine/Source/Runtime/VulkanRHI/Private/VulkanPipeline.h
+++ b/Engine/Source/Runtime/VulkanRHI/Private/VulkanPipeline.h
@@ -624,6 +624,22 @@ private:
 	};
 };
 
+struct VulkanResourceFrameCounter
+{
+	void Set(uint32_t Frame) const
+	{
+		LastUsedFrame.store(Frame, std::memory_order::relaxed);
+	}
+	
+	uint32 Get() const
+	{
+		return LastUsedFrame.load(std::memory_order::relaxed);
+	}
+
+	// Note that the frame counter may be updated from any number of threads, and is not entirely synchronized
+	mutable std::atomic<uint32> LastUsedFrame{0};
+};
+
 // Common pipeline class
 class FVulkanPipeline
 {
@@ -657,6 +673,9 @@ public:
 		}
 	}
 #endif
+
+	VulkanResourceFrameCounter FrameCounter;
+	
 protected:
 	FVulkanDevice* Device;
 	VkPipeline Pipeline;
@@ -786,6 +805,8 @@ public:
 	uint64	PrecacheKey;// hash of elements relevant to the PSO cache
 	FVulkanPSOKey							VulkanKey;
 
+	VulkanResourceFrameCounter FrameCounter;
+
 #if VULKAN_PSO_CACHE_DEBUG
 	FPixelShaderRHIRef					PixelShaderRHI;
 	FVertexShaderRHIRef					VertexShaderRHI;
diff --git a/Engine/Source/Runtime/VulkanRHI/Private/VulkanPipelineState.cpp b/Engine/Source/Runtime/VulkanRHI/Private/VulkanPipelineState.cpp
--- a/Engine/Source/Runtime/VulkanRHI/Private/VulkanPipelineState.cpp
+++ b/Engine/Source/Runtime/VulkanRHI/Private/VulkanPipelineState.cpp
@@ -545,6 +545,8 @@ void FVulkanCommandListContext::RHISetGraphicsPipelineState(FRHIGraphicsPipeline
 	FVulkanPipelineStateCacheManager* PipelineStateCache = Device.GetPipelineStateCache();
 	PipelineStateCache->LRUTouch(Pipeline);
 
+	Pipeline->FrameCounter.Set(GFrameNumberRenderThread);
+
 	FVulkanCommandBuffer& CommandBuffer = GetCommandBuffer();
 	bool bForceResetPipeline = !CommandBuffer.bHasPipeline;
 
@@ -586,6 +588,8 @@ void FVulkanCommandListContext::RHISetComputePipelineState(FRHIComputePipelineSt
 	FVulkanComputePipeline* ComputePipeline = ResourceCast(ComputePipelineState);
 	PendingComputeState->SetComputePipeline(ComputePipeline);
 
+	ComputePipeline->FrameCounter.Set(GFrameNumberRenderThread);
+	
 	ApplyStaticUniformBuffers(ResourceCast(ComputePipeline->GetComputeShader()));
 }
 
diff --git a/Engine/Source/Runtime/VulkanRHI/Private/VulkanRayTracing.cpp b/Engine/Source/Runtime/VulkanRHI/Private/VulkanRayTracing.cpp
--- a/Engine/Source/Runtime/VulkanRHI/Private/VulkanRayTracing.cpp
+++ b/Engine/Source/Runtime/VulkanRHI/Private/VulkanRayTracing.cpp
@@ -1760,6 +1760,11 @@ const FVulkanRayTracingShader* FVulkanRayTracingPipelineState::GetVulkanShader(E
 	return GetShaderData(Frequency).Shaders[ShaderIndex].GetReference();
 }
 
+int32 FVulkanRayTracingPipelineState::GetVulkanShaderNum(EShaderFrequency Frequency) const
+{
+	return GetShaderData(Frequency).Shaders.Num();
+}
+
 const TArray<uint8>& FVulkanRayTracingPipelineState::GetShaderHandles(EShaderFrequency Frequency) const
 {
 	return GetShaderData(Frequency).ShaderHandles;
@@ -1988,6 +1993,8 @@ void FVulkanCommandListContext::RHIRayTraceDispatch(
 	const FVulkanRayTracingPipelineState* Pipeline = ResourceCast(InRayTracingPipelineState);
 	FVulkanRayTracingShader* RayGenShader = ResourceCast(InRayGenShader);
 	FVulkanRayTracingShaderTable* ShaderTable = ResourceCast(InSBT);
+	
+	Pipeline->FrameCounter.Set(GFrameNumberRenderThread);
 
 	FVulkanCommandBuffer* const CmdBuffer = GetActiveCmdBuffer();
 	VulkanRHI::vkCmdBindPipeline(CmdBuffer->GetHandle(), VK_PIPELINE_BIND_POINT_RAY_TRACING_KHR, Pipeline->GetPipeline());
diff --git a/Engine/Source/Runtime/VulkanRHI/Private/VulkanRayTracing.h b/Engine/Source/Runtime/VulkanRHI/Private/VulkanRayTracing.h
--- a/Engine/Source/Runtime/VulkanRHI/Private/VulkanRayTracing.h
+++ b/Engine/Source/Runtime/VulkanRHI/Private/VulkanRayTracing.h
@@ -5,6 +5,7 @@
 #include "VulkanRHIPrivate.h"
 
 #include "RayTracingBuiltInResources.h"
+#include "VulkanPipeline.h"
 
 class FVulkanCommandListContext;
 class FVulkanBuffer;
@@ -286,6 +287,7 @@ public:
 
 	int32 GetShaderIndex(const FVulkanRayTracingShader* Shader) const;
 	const FVulkanRayTracingShader* GetVulkanShader(EShaderFrequency Frequency, int32 ShaderIndex) const;
+	int32 GetVulkanShaderNum(EShaderFrequency Frequency) const;
 	const TArray<uint8>& GetShaderHandles(EShaderFrequency Frequency) const;
 
 private:
@@ -309,6 +311,8 @@ public:
 	UE_DEPRECATED(5.5, "bAllowHitGroupIndexing is now stored in the ShaderBindingTable.")
 	bool bAllowHitGroupIndexing = true;
 
+	VulkanResourceFrameCounter FrameCounter;
+
 	friend FVulkanCommandListContext;
 	friend FVulkanRayTracingScene;
 };
diff --git a/Engine/Source/Runtime/VulkanRHI/Private/VulkanResources.h b/Engine/Source/Runtime/VulkanRHI/Private/VulkanResources.h
--- a/Engine/Source/Runtime/VulkanRHI/Private/VulkanResources.h
+++ b/Engine/Source/Runtime/VulkanRHI/Private/VulkanResources.h
@@ -149,7 +149,7 @@ public:
 	}
 
 	// Name should be pointing to "main_"
-	void GetEntryPoint(ANSICHAR* Name, int32 NameLength)
+	void GetEntryPoint(ANSICHAR* Name, int32 NameLength) const
 	{
 		FCStringAnsi::Snprintf(Name, NameLength, "main_%0.8x_%0.8x", SpirvContainer.GetSizeBytes(), CodeHeader.SpirvCRC);
 	}
@@ -180,7 +180,7 @@ public:
 		TArrayView<uint32> GetCodeView() {return CodeView;}
 	};
 
-	inline FSpirvCode GetSpirvCode()
+	inline FSpirvCode GetSpirvCode() const
 	{
 		return GetSpirvCode(SpirvContainer);
 	}
diff --git a/Engine/Source/Runtime/VulkanRHI/Private/VulkanUtil.cpp b/Engine/Source/Runtime/VulkanRHI/Private/VulkanUtil.cpp
--- a/Engine/Source/Runtime/VulkanRHI/Private/VulkanUtil.cpp
+++ b/Engine/Source/Runtime/VulkanRHI/Private/VulkanUtil.cpp
@@ -13,9 +13,12 @@
 #include "Misc/App.h"
 #include "Misc/OutputDeviceRedirector.h"
 #include "RHIValidationContext.h"
+#include "PipelineStateCache.h"
+#include "Async/ParallelFor.h"
 #include "HAL/FileManager.h"
 #include "RenderCore.h"
 #include "RHICoreNvidiaAftermath.h"
+#include "VulkanRayTracing.h"
 
 FVulkanDynamicRHI*	GVulkanRHI = nullptr;
 
@@ -680,6 +683,146 @@ void AftermathResolveMarkerCallback(const void* Marker, void** ResolvedMarkerDat
 	}
 #endif // VULKAN_SUPPORTS_NV_DIAGNOSTICS
 }
+
+void ConditionallyAddShader(const FVulkanShader* Value, TArray<const FVulkanShader*>& Shaders, TSet<const FVulkanShader*>& ShaderSet)
+{
+	if (Value && !ShaderSet.Contains(Value))
+	{
+		Shaders.Add(Value);
+		ShaderSet.Add(Value);
+	}
+}
+
+bool IsPipelineUnused(const VulkanResourceFrameCounter& Counter, uint32_t Threshold)
+{
+	// Check in case that it's ahead of the adapter, just in case it entered a strange state
+	const uint32 PipelineFrameIndex = Counter.Get();
+	return (PipelineFrameIndex <= GFrameNumberRenderThread) && (GFrameNumberRenderThread - PipelineFrameIndex) > Threshold;
+}
+
+void AftermathLateAssociate(float TimeLimitSeconds, uint32 FrameLimit)
+{
+	uint64 CycleStart = FPlatformTime::Cycles64();
+
+	UE_LOG(LogVulkanRHI, Log, TEXT("Starting late shader associations..."));
+
+	TArray<const FVulkanShader*> Shaders;
+	TSet<const FVulkanShader*>   ShaderSet;
+
+	uint32_t IgnoredPipelines = 0;
+
+	// Get active pipelines, allow one second for consolidation to finish
+	TArray<TRefCountPtr<FRHIResource>> PipelineResources;
+	PipelineStateCache::GetPipelineStates(PipelineResources, true, UE::FTimeout(FTimespan::FromSeconds(1)));
+
+	// Deduplicate shaders, Aftermath hashes are not local to the parent pipeline
+	for (FRHIResource* Resource : PipelineResources)
+	{
+		if (!Resource)
+		{
+			continue;
+		}
+		
+		switch (Resource->GetType())
+		{
+			default:
+				checkNoEntry();
+			case RRT_GraphicsPipelineState:
+			{
+				auto* Pipeline = static_cast<FVulkanRHIGraphicsPipelineState*>(Resource);
+					
+				if (IsPipelineUnused(Pipeline->FrameCounter, FrameLimit))
+				{
+					IgnoredPipelines++;
+					continue;
+				}
+					
+				ConditionallyAddShader(static_cast<FVulkanVertexShader*>(Pipeline->GetShader(SF_Vertex)), Shaders, ShaderSet);
+				ConditionallyAddShader(static_cast<FVulkanGeometryShader*>(Pipeline->GetShader(SF_Geometry)), Shaders, ShaderSet);
+				ConditionallyAddShader(static_cast<FVulkanGeometryShader*>(Pipeline->GetShader(SF_Amplification)), Shaders, ShaderSet);
+				ConditionallyAddShader(static_cast<FVulkanGeometryShader*>(Pipeline->GetShader(SF_Mesh)), Shaders, ShaderSet);
+				ConditionallyAddShader(static_cast<FVulkanPixelShader*>(Pipeline->GetShader(SF_Pixel)), Shaders, ShaderSet);
+				break;
+			}
+			case RRT_ComputePipelineState:
+			{
+				auto* Pipeline = static_cast<FVulkanComputePipeline*>(Resource);
+				
+				if (IsPipelineUnused(Pipeline->FrameCounter, FrameLimit))
+				{
+					IgnoredPipelines++;
+					continue;
+				}
+					
+				ConditionallyAddShader(static_cast<FVulkanComputeShader*>(Pipeline->GetComputeShader()), Shaders, ShaderSet);
+				break;
+			}
+			case RRT_RayTracingPipelineState:
+			{
+				auto* Pipeline = static_cast<FVulkanRayTracingPipelineState*>(Resource);
+					
+				if (IsPipelineUnused(Pipeline->FrameCounter, FrameLimit))
+				{
+					IgnoredPipelines++;
+					continue;
+				}
+
+				for (int32 i = 0; i < Pipeline->GetVulkanShaderNum(SF_RayGen); i++)
+				{
+					ConditionallyAddShader(Pipeline->GetVulkanShader(SF_RayGen, i), Shaders, ShaderSet);
+				}
+
+				for (int32 i = 0; i < Pipeline->GetVulkanShaderNum(SF_RayCallable); i++)
+				{
+					ConditionallyAddShader(Pipeline->GetVulkanShader(SF_RayCallable, i), Shaders, ShaderSet);
+				}
+
+				for (int32 i = 0; i < Pipeline->GetVulkanShaderNum(SF_RayHitGroup); i++)
+				{
+					ConditionallyAddShader(Pipeline->GetVulkanShader(SF_RayHitGroup, i), Shaders, ShaderSet);
+				}
+
+				for (int32 i = 0; i < Pipeline->GetVulkanShaderNum(SF_RayMiss); i++)
+				{
+					ConditionallyAddShader(Pipeline->GetVulkanShader(SF_RayMiss, i), Shaders, ShaderSet);
+				}
+				break;
+			}
+		}
+	}
+
+	UE_LOG(LogVulkanRHI, Log, TEXT("Late shader associations ignored %u pipelines based on frame fences"), IgnoredPipelines);
+
+	// Parallelize as much as possible to avoid timeouts
+	ParallelFor(Shaders.Num(), [CycleStart, TimeLimitSeconds, &Shaders](int32 Index)
+	{
+		// Aftermath handling is time constrained, if we hit the limit just stop
+		float Elapsed = FPlatformTime::ToSeconds64(FPlatformTime::Cycles64() - CycleStart);
+		if (Elapsed >= TimeLimitSeconds)
+		{
+			UE_CALL_ONCE([Elapsed]
+			{
+				UE_LOG(LogVulkanRHI, Warning, TEXT("Late shader associations timed out at %f.5ms"), Elapsed);
+			});
+			return;
+		}
+
+		const FVulkanShader*      Shader = Shaders[Index];
+		FVulkanShader::FSpirvCode Code   = Shader->GetSpirvCode();
+
+		ANSICHAR EntryPointName[1024];
+		Shader->GetEntryPoint(EntryPointName, 1024);
+
+		UE::RHICore::Nvidia::Aftermath::RegisterShaderBinary(
+			Code.GetCodeView().GetData(),
+			Code.GetCodeView().Num() * sizeof(uint32_t),
+			FString(EntryPointName)
+		);
+	});
+	
+	double TimeMS = FPlatformTime::ToMilliseconds64(FPlatformTime::Cycles64() - CycleStart);
+	UE_LOG(LogVulkanRHI, Log, TEXT("Created late shader associations, took %f.5ms"), TimeMS);
+}
 #endif
 
 #endif // (RHI_NEW_GPU_PROFILER == 0)
