Only output GPU assert message when a GPU crash happens.

diff --git a/Engine/Source/Runtime/RHICore/Private/RHIDiagnosticBuffer.cpp b/Engine/Source/Runtime/RHICore/Private/RHIDiagnosticBuffer.cpp
--- a/Engine/Source/Runtime/RHICore/Private/RHIDiagnosticBuffer.cpp
+++ b/Engine/Source/Runtime/RHICore/Private/RHIDiagnosticBuffer.cpp
@@ -30,29 +30,29 @@ FString FRHIDiagnosticBuffer::GetShaderDiagnosticMessages(uint32 DeviceIndex, ui
 			else
 			{
 				LanesString += FString::Printf(TEXT("Shader assertion failed - ID: 0x%08X (%d)"), Lane.MessageID, Lane.MessageID);
-			}
 
-			{
-				const int32* Payload = Lane.Payload.AsInt;
-				if (Payload[0] < 0 || Payload[1] < 0 || Payload[2] < 0 || Payload[3] < 0)
 				{
-					LanesString += FString::Printf(TEXT("\r\n\t\t\tPayload [ int32]: %d %d %d %d"), Payload[0], Payload[1], Payload[2], Payload[3]);
+					const int32* Payload = Lane.Payload.AsInt;
+					if (Payload[0] < 0 || Payload[1] < 0 || Payload[2] < 0 || Payload[3] < 0)
+					{
+						LanesString += FString::Printf(TEXT("\r\n\t\t\tPayload [ int32]: %d %d %d %d"), Payload[0], Payload[1], Payload[2], Payload[3]);
+					}
 				}
-			}
 
-			{
-				const uint32* Payload = Lane.Payload.AsUint;
-				LanesString += FString::Printf(TEXT("\r\n\t\t\tPayload [uint32]: %u %u %u %u"), Payload[0], Payload[1], Payload[2], Payload[3]);
-			}
+				{
+					const uint32* Payload = Lane.Payload.AsUint;
+					LanesString += FString::Printf(TEXT("\r\n\t\t\tPayload [uint32]: %u %u %u %u"), Payload[0], Payload[1], Payload[2], Payload[3]);
+				}
 
-			{
-				const uint32* Payload = Lane.Payload.AsUint;
-				LanesString += FString::Printf(TEXT("\r\n\t\t\tPayload [   hex]: 0x%08X 0x%08X 0x%08X 0x%08X"), Payload[0], Payload[1], Payload[2], Payload[3]);
-			}
+				{
+					const uint32* Payload = Lane.Payload.AsUint;
+					LanesString += FString::Printf(TEXT("\r\n\t\t\tPayload [   hex]: 0x%08X 0x%08X 0x%08X 0x%08X"), Payload[0], Payload[1], Payload[2], Payload[3]);
+				}
 
-			{
-				const float* Payload = Lane.Payload.AsFloat;
-				LanesString += FString::Printf(TEXT("\r\n\t\t\tPayload [ float]: %f %f %f %f"), Payload[0], Payload[1], Payload[2], Payload[3]);
+				{
+					const float* Payload = Lane.Payload.AsFloat;
+					LanesString += FString::Printf(TEXT("\r\n\t\t\tPayload [ float]: %f %f %f %f"), Payload[0], Payload[1], Payload[2], Payload[3]);
+				}
 			}
 		}
 		else
