Add support for defragging RHI Pool allocator for BLAS memory to improve indirect tracing gpu cost

diff --git a/Engine/Source/Runtime/D3D12RHI/Private/D3D12PoolAllocator.cpp b/Engine/Source/Runtime/D3D12RHI/Private/D3D12PoolAllocator.cpp
--- a/Engine/Source/Runtime/D3D12RHI/Private/D3D12PoolAllocator.cpp
+++ b/Engine/Source/Runtime/D3D12RHI/Private/D3D12PoolAllocator.cpp
@@ -763,31 +763,8 @@ void FD3D12PoolAllocator::CleanUpAllocations(uint64 InFrameLag, bool bForceFree)
 		}
 	}
 
-	// Update the allocation order
-	Algo::Sort(PoolAllocationOrder, [this](uint32 InLHS, uint32 InRHS)
-		{
-			// first allocate from 'fullest' pools when defrag is enabled, otherwise try and allocate from
-			// default sized pools with stable order
-			if (bDefragEnabled)
-			{
-				uint32 LHSUsePoolSize = Pools[InLHS] ? Pools[InLHS]->GetUsedSize() : UINT32_MAX;
-				uint32 RHSUsePoolSize = Pools[InRHS] ? Pools[InRHS]->GetUsedSize() : UINT32_MAX;
-				return LHSUsePoolSize > RHSUsePoolSize;
-			}
-			else
-			{
-				uint32 LHSPoolSize = Pools[InLHS] ? Pools[InLHS]->GetPoolSize() : UINT32_MAX;
-				uint32 RHSPoolSize = Pools[InRHS] ? Pools[InRHS]->GetPoolSize() : UINT32_MAX;
-				if (LHSPoolSize != RHSPoolSize)
-				{
-					return LHSPoolSize < RHSPoolSize;
-				}
-
-				uint32 LHSPoolIndex = Pools[InLHS] ? Pools[InLHS]->GetPoolIndex() : UINT32_MAX;
-				uint32 RHSPoolIndex = Pools[InRHS] ? Pools[InRHS]->GetPoolIndex() : UINT32_MAX;
-				return LHSPoolIndex < RHSPoolIndex;
-			}
-		});
+	// Sort the pools
+	SortPools();
 }
 
 
diff --git a/Engine/Source/Runtime/RHICore/Private/RHIPoolAllocator.cpp b/Engine/Source/Runtime/RHICore/Private/RHIPoolAllocator.cpp
--- a/Engine/Source/Runtime/RHICore/Private/RHIPoolAllocator.cpp
+++ b/Engine/Source/Runtime/RHICore/Private/RHIPoolAllocator.cpp
@@ -728,6 +728,35 @@ void FRHIPoolAllocator::DeallocateInternal(FRHIPoolAllocationData& AllocationDat
 }
 
 
+void FRHIPoolAllocator::SortPools()
+{
+	Algo::Sort(PoolAllocationOrder, [this](uint32 InLHS, uint32 InRHS)
+		{
+			// first allocate from 'fullest' pools when defrag is enabled, otherwise try and allocate from
+			// default sized pools with stable order
+			if (bDefragEnabled)
+			{
+				uint32 LHSUsePoolSize = Pools[InLHS] ? Pools[InLHS]->GetUsedSize() : UINT32_MAX;
+				uint32 RHSUsePoolSize = Pools[InRHS] ? Pools[InRHS]->GetUsedSize() : UINT32_MAX;
+				return LHSUsePoolSize > RHSUsePoolSize;
+			}
+			else
+			{
+				uint32 LHSPoolSize = Pools[InLHS] ? Pools[InLHS]->GetPoolSize() : UINT32_MAX;
+				uint32 RHSPoolSize = Pools[InRHS] ? Pools[InRHS]->GetPoolSize() : UINT32_MAX;
+				if (LHSPoolSize != RHSPoolSize)
+				{
+					return LHSPoolSize < RHSPoolSize;
+				}
+
+				uint32 LHSPoolIndex = Pools[InLHS] ? Pools[InLHS]->GetPoolIndex() : UINT32_MAX;
+				uint32 RHSPoolIndex = Pools[InRHS] ? Pools[InRHS]->GetPoolIndex() : UINT32_MAX;
+				return LHSPoolIndex < RHSPoolIndex;
+			}
+		});
+}
+
+
 void FRHIPoolAllocator::Defrag(FRHIContextArray const& Contexts, uint32 InMaxCopySize, uint32& CurrentCopySize)
 {
 	// Don't do anything when defrag is disabled for this allocator
diff --git a/Engine/Source/Runtime/RHICore/Public/RHIPoolAllocator.h b/Engine/Source/Runtime/RHICore/Public/RHIPoolAllocator.h
--- a/Engine/Source/Runtime/RHICore/Public/RHIPoolAllocator.h
+++ b/Engine/Source/Runtime/RHICore/Public/RHIPoolAllocator.h
@@ -239,6 +239,8 @@ protected:
 	// Supported allocator operations
 	RHICORE_API bool TryAllocateInternal(uint32 InSizeInBytes, uint32 InAllocationAlignment, ERHIPoolResourceTypes InAllocationResourceType, FRHIPoolAllocationData& AllocationData);
 	RHICORE_API void DeallocateInternal(FRHIPoolAllocationData& AllocationData);
+
+	RHICORE_API void SortPools();
 	
 	// Helper function to create a new platform specific pool
 	virtual FRHIMemoryPool* CreateNewPool(int16 InPoolIndex, uint32 InMinimumAllocationSize, ERHIPoolResourceTypes InAllocationResourceType) = 0;
