Make sure the free allocation and pooled allocation data objects are deleted whent the pool is destroyed - fixes memory leak

diff --git a/Engine/Source/Runtime/D3D12RHI/Private/D3D12PoolAllocator.cpp b/Engine/Source/Runtime/D3D12RHI/Private/D3D12PoolAllocator.cpp
--- a/Engine/Source/Runtime/D3D12RHI/Private/D3D12PoolAllocator.cpp
+++ b/Engine/Source/Runtime/D3D12RHI/Private/D3D12PoolAllocator.cpp
@@ -39,7 +39,6 @@ FD3D12MemoryPool::FD3D12MemoryPool(
 
 FD3D12MemoryPool::~FD3D12MemoryPool()
 {
-	Destroy();
 }
 
 
diff --git a/Engine/Source/Runtime/RHICore/Private/RHIPoolAllocator.cpp b/Engine/Source/Runtime/RHICore/Private/RHIPoolAllocator.cpp
--- a/Engine/Source/Runtime/RHICore/Private/RHIPoolAllocator.cpp
+++ b/Engine/Source/Runtime/RHICore/Private/RHIPoolAllocator.cpp
@@ -303,6 +303,7 @@ FRHIMemoryPool::~FRHIMemoryPool()
 	*/
 
 	Validate();
+	Destroy();
 }
 
 
@@ -333,6 +334,18 @@ void FRHIMemoryPool::Init()
 void FRHIMemoryPool::Destroy()
 {
 	// Validate state?
+
+	for (FRHIPoolAllocationData* FreeBlock : FreeBlocks)
+	{
+		ReleaseAllocationData(FreeBlock);
+	}
+	FreeBlocks.Empty();
+	
+	for (FRHIPoolAllocationData* AllocationData : AllocationDataPool)
+	{
+		delete AllocationData;
+	}
+	AllocationDataPool.Empty();
 }
 
 
