Aftermath, clarified hash names

diff --git a/Engine/Source/Runtime/RHICore/Private/RHICoreNvidiaAftermath.cpp b/Engine/Source/Runtime/RHICore/Private/RHICoreNvidiaAftermath.cpp
--- a/Engine/Source/Runtime/RHICore/Private/RHICoreNvidiaAftermath.cpp
+++ b/Engine/Source/Runtime/RHICore/Private/RHICoreNvidiaAftermath.cpp
@@ -321,7 +321,7 @@ namespace UE::RHICore::Nvidia::Aftermath
 		}
 
 		// Report binary hash separately, may not be the same
-		CrashResult.OutputLog += FString::Printf(TEXT("\n\t\t\tBinary Hash = %llu"), BinaryHash.hash);
+		CrashResult.OutputLog += FString::Printf(TEXT("\n\t\t\tAftermath Binary Hash = %llu"), BinaryHash.hash);
 
 		FShaderDebugInfo* Candidate = ShaderRegistration.ShaderBinaryToDebugHash.Find(BinaryHash.hash);
 		if (!Candidate)
@@ -332,12 +332,12 @@ namespace UE::RHICore::Nvidia::Aftermath
 
 		if (!Candidate->HashName.IsEmpty())
 		{
-			CrashResult.OutputLog += FString::Printf(TEXT("\n\t\t\tHash = %s"), *Candidate->HashName);
+			CrashResult.OutputLog += FString::Printf(TEXT("\n\t\t\tDXIL Hash = %s"), *Candidate->HashName);
 		}
 
 		if (!Candidate->DebugName.IsEmpty())
 		{
-			CrashResult.OutputLog += FString::Printf(TEXT("\n\t\t\tDebugName = '%s'"), *Candidate->DebugName);
+			CrashResult.OutputLog += FString::Printf(TEXT("\n\t\t\tDXIL DebugName = '%s'"), *Candidate->DebugName);
 		}
 	}
 
@@ -441,7 +441,7 @@ namespace UE::RHICore::Nvidia::Aftermath
 			}
 			
 			CrashResult.OutputLog += FString::Printf(TEXT("\n\t\t\tType = %s"), ShaderType);
-			CrashResult.OutputLog += FString::Printf(TEXT("\n\t\t\tHash = %llu"), Shader.shaderHash);
+			CrashResult.OutputLog += FString::Printf(TEXT("\n\t\t\tAftermath Hash = %llu"), Shader.shaderHash);
 
 			AppendGpuCrashShaderBinaryInfo(Decoder, Shader, CrashResult);
 		}
