Use critical section when accessing GShaderDiagnosticDatas because shaders can be loaded from multiple threads - contention should be pretty minimal since diagnostic data isn't used all that much

diff --git a/Engine/Source/Runtime/RHICore/Private/RHICoreShader.cpp b/Engine/Source/Runtime/RHICore/Private/RHICoreShader.cpp
--- a/Engine/Source/Runtime/RHICore/Private/RHICoreShader.cpp
+++ b/Engine/Source/Runtime/RHICore/Private/RHICoreShader.cpp
@@ -250,16 +250,17 @@ void UE::RHICore::SetupShaderDiagnosticData(FRHIShader* RHIShader, FShaderCodeRe
 	}
 }
 
+FCriticalSection GShaderDiagnosticDatasCS;
 TArray<FShaderDiagnosticData> GShaderDiagnosticDatas;
 void UE::RHICore::RegisterDiagnosticMessages(const TArray<FShaderDiagnosticData>& In)
 {
-	// Not thread safe
+	FScopeLock ScopeLock(&GShaderDiagnosticDatasCS);
 	GShaderDiagnosticDatas.Append(In);
 }
 
 const FString* UE::RHICore::GetDiagnosticMessage(uint32 MessageID)
 {
-	// Not thread safe
+	FScopeLock ScopeLock(&GShaderDiagnosticDatasCS);
 	if (const FShaderDiagnosticData* Found = GShaderDiagnosticDatas.FindByPredicate([MessageID](const FShaderDiagnosticData& In) { return In.Hash == MessageID; }))
 	{
 		return &Found->Message;
