Aftermath, improved debug names
- Print RHIShader debug name when available
- Scan for DXBC pipeline state validation entry point name

diff --git a/Engine/Source/Runtime/D3D12RHI/Private/D3D12NvidiaExtensions.cpp b/Engine/Source/Runtime/D3D12RHI/Private/D3D12NvidiaExtensions.cpp
--- a/Engine/Source/Runtime/D3D12RHI/Private/D3D12NvidiaExtensions.cpp
+++ b/Engine/Source/Runtime/D3D12RHI/Private/D3D12NvidiaExtensions.cpp
@@ -37,7 +37,7 @@ namespace UE::RHICore::Nvidia::Aftermath::D3D12
 	template<typename T>
 	void RegisterShader(T* Value)
 	{
-		RegisterShaderBinary(Value->Code.GetData(), Value->Code.Num());
+		RegisterShaderBinary(Value->Code.GetData(), Value->Code.Num(), Value->HasShaderName() ? Value->GetShaderName() : nullptr);
 	}
 	
 	template<typename T>
@@ -179,7 +179,7 @@ namespace UE::RHICore::Nvidia::Aftermath::D3D12
 			{
 				UE_CALL_ONCE([Elapsed]
 				{
-					UE_LOG(LogD3D12RHI, Warning, TEXT("Late shader associations timed out at %.0fms"), Elapsed);
+					UE_LOG(LogD3D12RHI, Warning, TEXT("Late shader associations timed out at %.0fs"), Elapsed);
 				});
 				return;
 			}
diff --git a/Engine/Source/Runtime/RHI/Public/RHIResources.h b/Engine/Source/Runtime/RHI/Public/RHIResources.h
--- a/Engine/Source/Runtime/RHI/Public/RHIResources.h
+++ b/Engine/Source/Runtime/RHI/Public/RHIResources.h
@@ -864,6 +864,8 @@ public:
 		FString ShaderName;
 		TArray<FName> UniformBufferNames;
 	} Debug;
+	
+	bool HasShaderName() const { return !Debug.ShaderName.IsEmpty(); } 
 
 	const TCHAR* GetShaderName() const
 	{
@@ -886,6 +888,7 @@ public:
 
 #else
 
+	bool HasShaderName() const { return false; }
 	const TCHAR* GetShaderName() const { return TEXT("<unknown>"); }
 	FString GetUniformBufferName(uint32 Index) const { return TEXT("<unknown>"); }
 
diff --git a/Engine/Source/Runtime/RHICore/Internal/ShaderBinaryUtilities.h b/Engine/Source/Runtime/RHICore/Internal/ShaderBinaryUtilities.h
--- a/Engine/Source/Runtime/RHICore/Internal/ShaderBinaryUtilities.h
+++ b/Engine/Source/Runtime/RHICore/Internal/ShaderBinaryUtilities.h
@@ -7,6 +7,13 @@
 
 namespace UE::ShaderBinaryUtilities
 {
+	struct FShaderDebugData
+	{
+		FString FunctionName;
+		FString PDBName;
+		FString Hash;
+	};
+	
 #if PLATFORM_SUPPORTS_PRAGMA_PACK
 #pragma pack(push, 1)
 #endif // PLATFORM_SUPPORTS_PRAGMA_PACK
@@ -40,6 +47,12 @@ namespace UE::ShaderBinaryUtilities
 			uint32 Flags;
 			uint8  Digest[16];
 		};
+
+		struct FPSVRevision3Info
+		{
+			uint8  Ignore[48];
+			uint32 EntryNameOffset;
+		};
 	}
 
 #if PLATFORM_SUPPORTS_PRAGMA_PACK
@@ -53,6 +66,18 @@ namespace UE::ShaderBinaryUtilities
 			
 		}
 
+		template<typename T>
+		T Get()
+		{
+			if (Offset + sizeof(T) > ByteSize)
+			{
+				checkf(false, TEXT("Parsing beyond EOS"));
+				return T{};
+			}
+			
+			return *reinterpret_cast<const T*>(Start + Offset);
+		}
+
 		template<typename T>
 		T Consume()
 		{
@@ -74,6 +99,11 @@ namespace UE::ShaderBinaryUtilities
 			return ctx;
 		}
 
+		void Skip(uint64 InOffset)
+		{
+			Offset += InOffset;
+		}
+
 		uint64 PendingBytes() const
 		{
 			return ByteSize - Offset;
@@ -89,12 +119,12 @@ namespace UE::ShaderBinaryUtilities
 		uint64 Offset = 0;
 	};
 	
-	static bool GetShaderBinaryDebugHashDXBC(const void* ShaderBinary, uint32 ByteSize, FString& OutDebugHash)
+	static void GetShaderBinaryDebugHashDXBC(const void* ShaderBinary, uint32 ByteSize, FShaderDebugData& OutData)
 	{
 		if (ByteSize < sizeof(DXBC::FHeader))
 		{
 			UE_LOG(LogRHICore, Error, TEXT("Shader byte size too small"));
-			return false;
+			return;
 		}
 
 		FParseContext Ctx(ShaderBinary, ByteSize);
@@ -123,16 +153,16 @@ namespace UE::ShaderBinaryUtilities
 					if (DXILDebugInfo.NameLength != HashLength)
 					{
 						UE_LOG(LogRHICore, Display, TEXT("DXIL name length not the expected hash"));
-						return false;
+						continue;
 					}
 
 					if (ChunkCtx.PendingBytes() < HashLength)
 					{
 						UE_LOG(LogRHICore, Display, TEXT("ILDN block corrupt"));
-						return false;
+						continue;
 					}
 
-					OutDebugHash = FString::ConstructFromPtrSize(static_cast<const ANSICHAR*>(ChunkCtx.Data()), HashLength);
+					OutData.PDBName = FString::ConstructFromPtrSize(static_cast<const ANSICHAR*>(ChunkCtx.Data()), HashLength);
 					break;
 				}
 				case 'HSAH':
@@ -143,15 +173,38 @@ namespace UE::ShaderBinaryUtilities
 					FUInt128 HashBytes;
 					FMemory::Memcpy(&HashBytes, DXILDebugInfo.Digest, sizeof(HashBytes));
 
-					OutDebugHash = BytesToHex(reinterpret_cast<const uint8*>(&HashBytes), sizeof(FUInt128));
+					OutData.Hash = BytesToHex(reinterpret_cast<const uint8*>(&HashBytes), sizeof(FUInt128));
+					break;
+				}
+				case '0VSP':
+				{
+					uint32 RevisionByteSize  = ChunkCtx.Consume<uint32>();
+					
+					if (RevisionByteSize >= sizeof(DXIL::FPSVRevision3Info))
+					{
+						DXIL::FPSVRevision3Info PSVInfo = ChunkCtx.Get<DXIL::FPSVRevision3Info>();
+						ChunkCtx.Skip(RevisionByteSize);
+
+						if (uint32 ResourceCount = ChunkCtx.Consume<uint32>())
+						{
+							uint32 BindInfoByteCount  = ChunkCtx.Consume<uint32>();
+							ChunkCtx.Skip(BindInfoByteCount * ResourceCount);
+						}
+						
+						uint32 StringTableSize  = ChunkCtx.Consume<uint32>();
+						if (StringTableSize <= PSVInfo.EntryNameOffset)
+						{
+							UE_LOG(LogRHICore, Error, TEXT("Pipeline state validation string table too small"));
+							continue;
+						}
+
+						ChunkCtx.Skip(PSVInfo.EntryNameOffset);
+						OutData.FunctionName = static_cast<const char*>(ChunkCtx.Data());
+					}
+						
 					break;
 				}
 			}
-			
-			return true;
 		}
-
-		// No relevant chunk found
-		return false;
 	}
 }
diff --git a/Engine/Source/Runtime/RHICore/Private/RHICoreNvidiaAftermath.cpp b/Engine/Source/Runtime/RHICore/Private/RHICoreNvidiaAftermath.cpp
--- a/Engine/Source/Runtime/RHICore/Private/RHICoreNvidiaAftermath.cpp
+++ b/Engine/Source/Runtime/RHICore/Private/RHICoreNvidiaAftermath.cpp
@@ -136,8 +136,8 @@ namespace UE::RHICore::Nvidia::Aftermath
 
 	struct FShaderDebugInfo
 	{
-		FString HashName;
 		FString DebugName;
+		ShaderBinaryUtilities::FShaderDebugData DebugData;
 	};
 	
 	struct FShaderRegistration
@@ -330,14 +330,24 @@ namespace UE::RHICore::Nvidia::Aftermath
 			return;
 		}
 
-		if (!Candidate->HashName.IsEmpty())
+		if (!Candidate->DebugName.IsEmpty())
 		{
-			CrashResult.OutputLog += FString::Printf(TEXT("\n\t\t\tDXIL Hash = %s"), *Candidate->HashName);
+			CrashResult.OutputLog += FString::Printf(TEXT("\n\t\t\tDebugName = '%s'"), *Candidate->DebugName);
 		}
 
-		if (!Candidate->DebugName.IsEmpty())
+		if (!Candidate->DebugData.FunctionName.IsEmpty())
+		{
+			CrashResult.OutputLog += FString::Printf(TEXT("\n\t\t\tFunction = %s"), *Candidate->DebugData.FunctionName);
+		}
+
+		if (!Candidate->DebugData.PDBName.IsEmpty())
+		{
+			CrashResult.OutputLog += FString::Printf(TEXT("\n\t\t\tShader PDB = %s"), *Candidate->DebugData.PDBName);
+		}
+
+		if (!Candidate->DebugData.Hash.IsEmpty())
 		{
-			CrashResult.OutputLog += FString::Printf(TEXT("\n\t\t\tDXIL DebugName = '%s'"), *Candidate->DebugName);
+			CrashResult.OutputLog += FString::Printf(TEXT("\n\t\t\tShader Hash = %s"), *Candidate->DebugData.Hash);
 		}
 	}
 
@@ -919,7 +929,7 @@ namespace UE::RHICore::Nvidia::Aftermath
 			}
 			case 'CBXD': // DXBC container header
 			{
-				ShaderBinaryUtilities::GetShaderBinaryDebugHashDXBC(ShaderBinary, ByteSize, DebugInfo.HashName);
+				ShaderBinaryUtilities::GetShaderBinaryDebugHashDXBC(ShaderBinary, ByteSize, DebugInfo.DebugData);
 					
 #ifdef __d3d12_h__
 				D3D12_SHADER_BYTECODE ByteCode;
diff --git a/Engine/Source/Runtime/VulkanRHI/Private/VulkanUtil.cpp b/Engine/Source/Runtime/VulkanRHI/Private/VulkanUtil.cpp
--- a/Engine/Source/Runtime/VulkanRHI/Private/VulkanUtil.cpp
+++ b/Engine/Source/Runtime/VulkanRHI/Private/VulkanUtil.cpp
@@ -414,11 +414,22 @@ void FVulkanGPUProfiler::EndFrame()
 #endif // (RHI_NEW_GPU_PROFILER == 0)
 
 #if NV_AFTERMATH
-void ConditionallyAddShader(const FVulkanShader* Value, TArray<const FVulkanShader*>& Shaders, TSet<const FVulkanShader*>& ShaderSet)
+struct FShaderPair
+{
+	const FRHIShader*    RHIShader = nullptr;
+	const FVulkanShader* VulkanShader = nullptr;
+};
+
+template<typename T>
+void ConditionallyAddShader(const T* Value, TArray<FShaderPair>& Shaders, TSet<const FRHIShader*>& ShaderSet)
 {
 	if (Value && !ShaderSet.Contains(Value))
 	{
-		Shaders.Add(Value);
+		Shaders.Add(FShaderPair {
+			static_cast<const FRHIShader*>(Value),
+			static_cast<const FVulkanShader*>(Value),
+		});
+		
 		ShaderSet.Add(Value);
 	}
 }
@@ -436,8 +447,8 @@ void AftermathLateAssociate(float TimeLimitSeconds, uint32 FrameLimit)
 
 	UE_LOG(LogVulkanRHI, Log, TEXT("Starting late shader associations..."));
 
-	TArray<const FVulkanShader*> Shaders;
-	TSet<const FVulkanShader*>   ShaderSet;
+	TArray<FShaderPair>     Shaders;
+	TSet<const FRHIShader*> ShaderSet;
 
 	uint32_t IgnoredPipelines = 0;
 
@@ -532,21 +543,30 @@ void AftermathLateAssociate(float TimeLimitSeconds, uint32 FrameLimit)
 		{
 			UE_CALL_ONCE([Elapsed]
 			{
-				UE_LOG(LogVulkanRHI, Warning, TEXT("Late shader associations timed out at %f.5ms"), Elapsed);
+				UE_LOG(LogVulkanRHI, Warning, TEXT("Late shader associations timed out at %.0fs"), Elapsed);
 			});
 			return;
 		}
 
-		const FVulkanShader*      Shader = Shaders[Index];
-		FVulkanShader::FSpirvCode Code   = Shader->GetSpirvCode();
-
-		ANSICHAR EntryPointName[1024];
-		Shader->GetEntryPoint(EntryPointName, 1024);
+		FShaderPair               Shader = Shaders[Index];
+		FVulkanShader::FSpirvCode Code   = Shader.VulkanShader->GetSpirvCode();
 
+		FString DebugName;
+		if (Shader.RHIShader->HasShaderName())
+		{
+			DebugName = Shader.RHIShader->GetShaderName();	
+		}
+		else
+		{
+			ANSICHAR AnsiName[1024];
+			Shader.VulkanShader->GetEntryPoint(AnsiName, 1024);
+			DebugName = AnsiName;
+		}
+		
 		UE::RHICore::Nvidia::Aftermath::RegisterShaderBinary(
 			Code.GetCodeView().GetData(),
 			Code.GetCodeView().Num() * sizeof(uint32_t),
-			FString(EntryPointName)
+			DebugName
 		);
 	});
 	
