Refactor RHIDescriptorAllocator to better support descriptor range allocations.
Minor changes to D3D12 descriptor heap resizing.

diff --git a/Engine/Source/Runtime/D3D12RHI/Private/D3D12BindlessDescriptors.cpp b/Engine/Source/Runtime/D3D12RHI/Private/D3D12BindlessDescriptors.cpp
--- a/Engine/Source/Runtime/D3D12RHI/Private/D3D12BindlessDescriptors.cpp
+++ b/Engine/Source/Runtime/D3D12RHI/Private/D3D12BindlessDescriptors.cpp
@@ -164,7 +164,12 @@ FRHIDescriptorHandle FD3D12BindlessDescriptorAllocator::AllocateDescriptor(ERHID
 {
 	if (DescriptorType == ERHIDescriptorType::Sampler)
 	{
-		FRHIDescriptorHandle Result = SamplerAllocator->Allocate(DescriptorType);
+		if (!AreSamplersBindless())
+		{
+			return FRHIDescriptorHandle();
+		}
+
+		const FRHIDescriptorHandle Result = SamplerAllocator->Allocate(DescriptorType);
 		check(Result.IsValid());
 		return Result;
 	}
@@ -179,35 +184,59 @@ FRHIDescriptorHandle FD3D12BindlessDescriptorAllocator::AllocateDescriptor(ERHID
 #if D3D12RHI_BINDLESS_RESOURCE_MANAGER_SUPPORTS_RESIZING
 	if (!Result.IsValid())
 	{
-		TRACE_CPUPROFILER_EVENT_SCOPE(FD3D12Adapter::BindlessResourceAllocateHandle(GrowHeap));
-
-		FScopeLock ScopeLock(&ResourceHeapsCS);
+		ResizeGrow();
+		Result = ResourceAllocator->Allocate(DescriptorType);
+	}
+#endif
 
-		// Grow the descriptor handle allocator
-		uint32 CurrentNumDescriptors = ResourceAllocator->GetCapacity();
-		uint32 NewNumDescriptors = FMath::Min<uint32>(CurrentNumDescriptors * 2, MaxResourceHeapSize);
+	check(Result.IsValid());
+	return Result;
+}
 
-		if (CurrentNumDescriptors == NewNumDescriptors)
+TOptional<FRHIDescriptorAllocation> FD3D12BindlessDescriptorAllocator::AllocateDescriptors(ERHIDescriptorType DescriptorType, uint32 DescriptorCount)
+{
+	if (DescriptorType == ERHIDescriptorType::Sampler)
+	{
+		if (!AreSamplersBindless())
 		{
-			UE_LOG(LogD3D12RHI, Fatal, TEXT("Hit D3D12 device limits on descriptors when attempting to allocate a larger descriptor heap."));
+			return TOptional<FRHIDescriptorAllocation>();
 		}
 
-		Result = ResourceAllocator->ResizeGrowAndAllocate(NewNumDescriptors, DescriptorType);
+		return SamplerAllocator->Allocate(DescriptorCount);
+	}
 
-		// Grow the CPU heaps for all devices
-		for (FD3D12Device* ParentDevice : GetParentAdapter()->GetDevices())
-		{
-			checkSlow(ParentDevice->GetBindlessDescriptorManager().GetResourceManager() != nullptr);
-			ParentDevice->GetBindlessDescriptorManager().GetResourceManager()->GrowCPUHeap(CurrentNumDescriptors, NewNumDescriptors);
-		}
+	if (!AreResourcesBindless())
+	{
+		return TOptional<FRHIDescriptorAllocation>();
+	}
 
-		return Result;
+	return ResourceAllocator->Allocate(DescriptorCount);
+}
+
+#if D3D12RHI_BINDLESS_RESOURCE_MANAGER_SUPPORTS_RESIZING
+void FD3D12BindlessDescriptorAllocator::ResizeGrow()
+{
+	TRACE_CPUPROFILER_EVENT_SCOPE(FD3D12Adapter::BindlessResourceAllocateHandle(GrowHeap));
+
+	FScopeLock ScopeLock(&ResourceHeapsCS);
+
+	// Grow the descriptor handle allocator
+	uint32 CurrentNumDescriptors = ResourceAllocator->GetCapacity();
+	uint32 NewNumDescriptors = FMath::Min<uint32>(CurrentNumDescriptors * 2, MaxResourceHeapSize);
+
+	if (CurrentNumDescriptors == NewNumDescriptors)
+	{
+		UE_LOG(LogD3D12RHI, Fatal, TEXT("Hit D3D12 device limits on descriptors when attempting to allocate a larger descriptor heap."));
 	}
-#endif
 
-	check(Result.IsValid());
-	return Result;
+	// Grow the CPU heaps for all devices
+	for (FD3D12Device* ParentDevice : GetParentAdapter()->GetDevices())
+	{
+		checkSlow(ParentDevice->GetBindlessDescriptorManager().GetResourceManager() != nullptr);
+		ParentDevice->GetBindlessDescriptorManager().GetResourceManager()->GrowCPUHeap(CurrentNumDescriptors, NewNumDescriptors);
+	}
 }
+#endif
 
 void FD3D12BindlessDescriptorAllocator::FreeDescriptor(FRHIDescriptorHandle Handle)
 {
diff --git a/Engine/Source/Runtime/D3D12RHI/Private/D3D12BindlessDescriptors.h b/Engine/Source/Runtime/D3D12RHI/Private/D3D12BindlessDescriptors.h
--- a/Engine/Source/Runtime/D3D12RHI/Private/D3D12BindlessDescriptors.h
+++ b/Engine/Source/Runtime/D3D12RHI/Private/D3D12BindlessDescriptors.h
@@ -19,6 +19,10 @@ struct FD3D12Payload;
 
 #include COMPILED_PLATFORM_HEADER(D3D12BindlessDescriptors.h)
 
+#if !defined(D3D12RHI_BINDLESS_RESOURCE_MANAGER_SUPPORTS_RESIZING)
+	#error D3D12RHI_BINDLESS_RESOURCE_MANAGER_SUPPORTS_RESIZING needs to be defined
+#endif
+
 namespace UE::D3D12BindlessDescriptors
 {
 	FD3D12DescriptorHeap* CreateCpuHeap(FD3D12Device* InDevice, ERHIDescriptorHeapType InType, uint32 InNewNumDescriptorsPerHeap);
@@ -42,8 +46,10 @@ public:
 
 	// Bindless descriptor allocators are stored in the adapter, so descriptor handles can be allocated once and shared for multi-GPU objects
 	FRHIDescriptorHandle AllocateDescriptor(ERHIDescriptorType DescriptorType);
+	TOptional<FRHIDescriptorAllocation> AllocateDescriptors(ERHIDescriptorType DescriptorType, uint32 DescriptorCount);
 
 	void FreeDescriptor(FRHIDescriptorHandle Handle);
+	void FreeDescriptors(ERHIDescriptorType DescriptorType, uint32 Offset);
 
 	FCriticalSection& GetResourceHeapsCS() { return ResourceHeapsCS; }
 
@@ -57,6 +63,10 @@ public:
 #endif
 
 private:
+#if D3D12RHI_BINDLESS_RESOURCE_MANAGER_SUPPORTS_RESIZING
+	void ResizeGrow();
+#endif
+
 	ERHIBindlessConfiguration BindlessConfiguration{};
 
 	uint32 MaxResourceHeapSize = 0;
@@ -103,8 +113,6 @@ private:
 
 #if !D3D12RHI_CUSTOM_BINDLESS_RESOURCE_MANAGER
 
-#define D3D12RHI_BINDLESS_RESOURCE_MANAGER_SUPPORTS_RESIZING 1
-
 // Helper container for all context related bindless state.
 struct FD3D12ContextBindlessState
 {
diff --git a/Engine/Source/Runtime/D3D12RHI/Private/D3D12Descriptors.cpp b/Engine/Source/Runtime/D3D12RHI/Private/D3D12Descriptors.cpp
--- a/Engine/Source/Runtime/D3D12RHI/Private/D3D12Descriptors.cpp
+++ b/Engine/Source/Runtime/D3D12RHI/Private/D3D12Descriptors.cpp
@@ -567,10 +567,9 @@ FD3D12DescriptorHeap* FD3D12DescriptorHeapManager::AllocateHeap(const TCHAR* InD
 	{
 		if (GlobalHeap.HandlesHeapAllocationsWithFlags(InHeapType, InHeapFlags))
 		{
-			uint32 Offset = 0;
-			if (GlobalHeap.Allocate(InNumDescriptors, Offset))
+			if (TOptional<FRHIDescriptorAllocation> Allocation = GlobalHeap.Allocate(InNumDescriptors))
 			{
-				return new FD3D12DescriptorHeap(GlobalHeap.GetHeap(), Offset, InNumDescriptors);
+				return new FD3D12DescriptorHeap(GlobalHeap.GetHeap(), Allocation->StartIndex, Allocation->Count);
 			}
 
 			// TODO: handle running out of space
@@ -610,7 +609,7 @@ void FD3D12DescriptorHeapManager::ImmediateFreeHeap(FD3D12DescriptorHeap* InHeap
 		{
 			if (GlobalHeap.IsHeapAChild(InHeap))
 			{
-				GlobalHeap.Free(InHeap->GetOffset(), InHeap->GetNumDescriptors());
+				GlobalHeap.Free(FRHIDescriptorAllocation(InHeap->GetOffset(), InHeap->GetNumDescriptors()));
 				InHeap->Release();
 				return;
 			}
diff --git a/Engine/Source/Runtime/D3D12RHI/Private/Windows/WindowsD3D12BindlessDescriptors.h b/Engine/Source/Runtime/D3D12RHI/Private/Windows/WindowsD3D12BindlessDescriptors.h
--- a/Engine/Source/Runtime/D3D12RHI/Private/Windows/WindowsD3D12BindlessDescriptors.h
+++ b/Engine/Source/Runtime/D3D12RHI/Private/Windows/WindowsD3D12BindlessDescriptors.h
@@ -3,3 +3,4 @@
 #pragma once
 
 #define D3D12RHI_CUSTOM_BINDLESS_RESOURCE_MANAGER 0
+#define D3D12RHI_BINDLESS_RESOURCE_MANAGER_SUPPORTS_RESIZING 1
diff --git a/Engine/Source/Runtime/RHICore/Private/RHIDescriptorAllocator.cpp b/Engine/Source/Runtime/RHICore/Private/RHIDescriptorAllocator.cpp
--- a/Engine/Source/Runtime/RHICore/Private/RHIDescriptorAllocator.cpp
+++ b/Engine/Source/Runtime/RHICore/Private/RHIDescriptorAllocator.cpp
@@ -33,7 +33,7 @@ void FRHIDescriptorAllocator::Shutdown()
 	Capacity = 0;
 }
 
-FRHIDescriptorHandle FRHIDescriptorAllocator::ResizeGrowAndAllocate(uint32 InNewNumDescriptors, ERHIDescriptorType InType)
+void FRHIDescriptorAllocator::ResizeGrow(uint32 InNewNumDescriptors)
 {
 	check(Capacity < InNewNumDescriptors);
 	
@@ -59,37 +59,15 @@ FRHIDescriptorHandle FRHIDescriptorAllocator::ResizeGrowAndAllocate(uint32 InNew
 	}
 
 	Capacity = InNewNumDescriptors;
-
-	uint32 Index{};
-	verify(Allocate(1, Index))
-	return FRHIDescriptorHandle(InType, Index);
-}
-
-FRHIDescriptorHandle FRHIDescriptorAllocator::Allocate(ERHIDescriptorType InType)
-{
-	uint32 Index{};
-	if (Allocate(1, Index))
-	{
-		return FRHIDescriptorHandle(InType, Index);
-	}
-	return FRHIDescriptorHandle();
-}
-
-void FRHIDescriptorAllocator::Free(FRHIDescriptorHandle InHandle)
-{
-	if (InHandle.IsValid())
-	{
-		Free(InHandle.GetIndex(), 1);
-	}
 }
 
-bool FRHIDescriptorAllocator::Allocate(uint32 NumDescriptors, uint32& OutSlot)
+TOptional<FRHIDescriptorAllocation> FRHIDescriptorAllocator::Allocate(uint32 NumDescriptors)
 {
 	FScopeLock Lock(&CriticalSection);
-	return AllocateInternal(NumDescriptors, OutSlot);
+	return AllocateInternal(NumDescriptors);
 }
 
-bool FRHIDescriptorAllocator::AllocateInternal(uint32 NumDescriptors, uint32& OutSlot)
+TOptional<FRHIDescriptorAllocation> FRHIDescriptorAllocator::AllocateInternal(uint32 NumDescriptors)
 {
 	if (const uint32 NumRanges = Ranges.Num(); NumRanges > 0)
 	{
@@ -110,22 +88,23 @@ bool FRHIDescriptorAllocator::AllocateInternal(uint32 NumDescriptors, uint32& Ou
 				{
 					CurrentRange.First += NumDescriptors;
 				}
-				OutSlot = First;
 
 				RecordAlloc(NumDescriptors);
 
-				return true;
+				return FRHIDescriptorAllocation(First, NumDescriptors);
 			}
 			++Index;
 		} while (Index < NumRanges);
 	}
 
-	OutSlot = UINT_MAX;
-	return false;
+	return TOptional<FRHIDescriptorAllocation>();
 }
 
-void FRHIDescriptorAllocator::Free(uint32 Offset, uint32 NumDescriptors)
+void FRHIDescriptorAllocator::Free(FRHIDescriptorAllocation Allocation)
 {
+	const uint32 Offset = Allocation.StartIndex;
+	const uint32 NumDescriptors = Allocation.Count;
+
 	if (Offset == UINT_MAX || NumDescriptors == 0)
 	{
 		return;
@@ -227,6 +206,23 @@ void FRHIDescriptorAllocator::Free(uint32 Offset, uint32 NumDescriptors)
 	}
 }
 
+FRHIDescriptorHandle FRHIDescriptorAllocator::Allocate(ERHIDescriptorType InType)
+{
+	if (TOptional<FRHIDescriptorAllocation> Allocation = Allocate(1))
+	{
+		return FRHIDescriptorHandle(InType, Allocation->StartIndex);
+	}
+	return FRHIDescriptorHandle();
+}
+
+void FRHIDescriptorAllocator::Free(FRHIDescriptorHandle InHandle)
+{
+	if (InHandle.IsValid())
+	{
+		Free(FRHIDescriptorAllocation(InHandle.GetIndex(), 1));
+	}
+}
+
 bool FRHIDescriptorAllocator::GetAllocatedRange(FRHIDescriptorAllocatorRange& OutRange)
 {
 	const uint32 VeryFirstIndex = 0;
@@ -287,14 +283,14 @@ void FRHIHeapDescriptorAllocator::Free(FRHIDescriptorHandle InHandle)
 	}
 }
 
-bool FRHIHeapDescriptorAllocator::Allocate(uint32 NumDescriptors, uint32& OutSlot)
+TOptional<FRHIDescriptorAllocation> FRHIHeapDescriptorAllocator::Allocate(uint32 NumDescriptors)
 {
-	return FRHIDescriptorAllocator::Allocate(NumDescriptors, OutSlot);
+	return FRHIDescriptorAllocator::Allocate(NumDescriptors);
 }
 
-void FRHIHeapDescriptorAllocator::Free(uint32 Slot, uint32 NumDescriptors)
+void FRHIHeapDescriptorAllocator::Free(FRHIDescriptorAllocation Allocation)
 {
-	FRHIDescriptorAllocator::Free(Slot, NumDescriptors);
+	FRHIDescriptorAllocator::Free(Allocation);
 }
 
 ////////////////////////////////////////////////////////////////////////////////////////////////////////////////////////
diff --git a/Engine/Source/Runtime/RHICore/Public/RHIDescriptorAllocator.h b/Engine/Source/Runtime/RHICore/Public/RHIDescriptorAllocator.h
--- a/Engine/Source/Runtime/RHICore/Public/RHIDescriptorAllocator.h
+++ b/Engine/Source/Runtime/RHICore/Public/RHIDescriptorAllocator.h
@@ -13,6 +13,18 @@ struct FRHIDescriptorAllocatorRange
 	uint32 Last;
 };
 
+struct FRHIDescriptorAllocation
+{
+	FRHIDescriptorAllocation(uint32 InStartIndex, uint32 InCount)
+		: StartIndex(InStartIndex)
+		, Count(InCount)
+	{
+	}
+
+	uint32 StartIndex = 0;
+	uint32 Count = 0;
+};
+
 class FRHIDescriptorAllocator
 {
 public:
@@ -23,14 +35,14 @@ public:
 	RHICORE_API void Init(uint32 InNumDescriptors, TConstArrayView<TStatId> InStats);
 	RHICORE_API void Shutdown();
 
-	RHICORE_API FRHIDescriptorHandle ResizeGrowAndAllocate(uint32 InNewNumDescriptors, ERHIDescriptorType InType);
+	RHICORE_API void ResizeGrow(uint32 InNewNumDescriptors);
+
+	RHICORE_API TOptional<FRHIDescriptorAllocation> Allocate(uint32 NumDescriptors);
+	RHICORE_API void Free(FRHIDescriptorAllocation Allocation);
 
 	RHICORE_API FRHIDescriptorHandle Allocate(ERHIDescriptorType InType);
 	RHICORE_API void Free(FRHIDescriptorHandle InHandle);
 
-	RHICORE_API bool Allocate(uint32 NumDescriptors, uint32& OutSlot);
-	RHICORE_API void Free(uint32 Slot, uint32 NumDescriptors);
-
 	// Get the range of allocated descriptors. Useful for determining the smallest range to copy between heaps.
 	RHICORE_API bool GetAllocatedRange(FRHIDescriptorAllocatorRange& OutRange);
 
@@ -38,7 +50,7 @@ public:
 
 private:
 
-	RHICORE_API bool AllocateInternal(uint32 NumDescriptors, uint32& OutSlot);
+	RHICORE_API TOptional<FRHIDescriptorAllocation> AllocateInternal(uint32 NumDescriptors);
 
 	void RecordAlloc(uint32 Count)
 	{
@@ -80,12 +92,12 @@ public:
 	RHICORE_API FRHIDescriptorHandle Allocate(ERHIDescriptorType InType);
 	RHICORE_API void Free(FRHIDescriptorHandle InHandle);
 
-	RHICORE_API bool Allocate(uint32 NumDescriptors, uint32& OutSlot);
-	RHICORE_API void Free(uint32 Slot, uint32 NumDescriptors);
+	RHICORE_API TOptional<FRHIDescriptorAllocation> Allocate(uint32 NumDescriptors);
+	RHICORE_API void Free(FRHIDescriptorAllocation Allocation);
 
 	using FRHIDescriptorAllocator::GetAllocatedRange;
 	using FRHIDescriptorAllocator::GetCapacity;
-	using FRHIDescriptorAllocator::ResizeGrowAndAllocate;
+	using FRHIDescriptorAllocator::ResizeGrow;
 
 	ERHIDescriptorTypeMask GetTypeMask() const
 	{
