Add an option to enable Aftermath shader error reporting.

The new CVar `r.GPUCrashDebugging.Aftermath.ShaderErrorReporting` or command line option `-nvAftermathShaderErrorReporting` enables the GFSDK_Aftermath_FeatureFlags_EnableShaderErrorReporting  flag.

#rnx

diff --git a/Engine/Source/Runtime/RHICore/Private/RHICoreNvidiaAftermath.cpp b/Engine/Source/Runtime/RHICore/Private/RHICoreNvidiaAftermath.cpp
--- a/Engine/Source/Runtime/RHICore/Private/RHICoreNvidiaAftermath.cpp
+++ b/Engine/Source/Runtime/RHICore/Private/RHICoreNvidiaAftermath.cpp
@@ -80,6 +80,13 @@ namespace UE::RHICore::Nvidia::Aftermath
 		ECVF_ReadOnly
 	);
 
+	static TAutoConsoleVariable<int32> CVarAftermath_ShaderErrorReporting(
+		TEXT("r.GPUCrashDebugging.Aftermath.ShaderErrorReporting"),
+		0,
+		TEXT("Enable shader error reporting for Aftermath dumps"),
+		ECVF_ReadOnly
+	);
+
 	static TAutoConsoleVariable<int32> CVarAftermath_TrackAll(
 		TEXT("r.GPUCrashDebugging.Aftermath.TrackAll"),
 		0,
@@ -223,16 +230,16 @@ namespace UE::RHICore::Nvidia::Aftermath
 			| GFSDK_Aftermath_FeatureFlags_EnableResourceTracking
 			| GFSDK_Aftermath_FeatureFlags_CallStackCapturing
 			| GFSDK_Aftermath_FeatureFlags_GenerateShaderDebugInfo
-			// @todo - GFSDK_Aftermath_FeatureFlags_EnableShaderErrorReporting is disabled to prevent TDRs until Nvidia fixes this
-			//| GFSDK_Aftermath_FeatureFlags_EnableShaderErrorReporting
+			| GFSDK_Aftermath_FeatureFlags_EnableShaderErrorReporting
 			;
 
 		Flags = GFSDK_Aftermath_FeatureFlags_Minimum;
-		Flags |= (FParse::Param(FCommandLine::Get(), TEXT("nvaftermathmarkers"            )) || (CVarAftermath_Markers            ->GetInt())) ? GFSDK_Aftermath_FeatureFlags_EnableMarkers           : 0;
-		Flags |= (FParse::Param(FCommandLine::Get(), TEXT("nvaftermathresources"          )) || (CVarAftermath_Resources          ->GetInt())) ? GFSDK_Aftermath_FeatureFlags_EnableResourceTracking  : 0;
-		Flags |= (FParse::Param(FCommandLine::Get(), TEXT("nvaftermathcallstack"          )) || (CVarAftermath_Callstack          ->GetInt())) ? GFSDK_Aftermath_FeatureFlags_CallStackCapturing      : 0;
-		Flags |= (FParse::Param(FCommandLine::Get(), TEXT("nvAftermathDumpShaderDebugInfo")) || (CVarAftermath_DumpShaderDebugInfo->GetInt())) ? GFSDK_Aftermath_FeatureFlags_GenerateShaderDebugInfo : 0;
-		Flags |= (FParse::Param(FCommandLine::Get(), TEXT("nvaftermathall"                )) || (CVarAftermath_TrackAll           ->GetInt())) ? AllFlags                                             : 0;
+		Flags |= (FParse::Param(FCommandLine::Get(), TEXT("nvaftermathmarkers"             ))  || (CVarAftermath_Markers             ->GetInt())) ? GFSDK_Aftermath_FeatureFlags_EnableMarkers              : 0;
+		Flags |= (FParse::Param(FCommandLine::Get(), TEXT("nvaftermathresources"           ))  || (CVarAftermath_Resources           ->GetInt())) ? GFSDK_Aftermath_FeatureFlags_EnableResourceTracking     : 0;
+		Flags |= (FParse::Param(FCommandLine::Get(), TEXT("nvaftermathcallstack"           ))  || (CVarAftermath_Callstack           ->GetInt())) ? GFSDK_Aftermath_FeatureFlags_CallStackCapturing         : 0;
+		Flags |= (FParse::Param(FCommandLine::Get(), TEXT("nvAftermathDumpShaderDebugInfo" ))  || (CVarAftermath_DumpShaderDebugInfo ->GetInt())) ? GFSDK_Aftermath_FeatureFlags_GenerateShaderDebugInfo    : 0;
+		Flags |= (FParse::Param(FCommandLine::Get(), TEXT("nvAftermathShaderErrorReporting"))  || (CVarAftermath_ShaderErrorReporting->GetInt())) ? GFSDK_Aftermath_FeatureFlags_EnableShaderErrorReporting : 0;
+		Flags |= (FParse::Param(FCommandLine::Get(), TEXT("nvaftermathall"                 ))  || (CVarAftermath_TrackAll            ->GetInt())) ? AllFlags                                                : 0;
 
 		GFSDK_Aftermath_Result Result = GFSDK_Aftermath_EnableGpuCrashDumps(
 			GFSDK_Aftermath_Version_API,
