Changing UpdateGenericResourceCollection to lock/unlock its entire buffer when updating. Not all of our platforms support lock with offset yet, so this is the next best approach.

diff --git a/Engine/Source/Runtime/RHICore/Internal/RHICoreResourceCollection.h b/Engine/Source/Runtime/RHICore/Internal/RHICoreResourceCollection.h
--- a/Engine/Source/Runtime/RHICore/Internal/RHICoreResourceCollection.h
+++ b/Engine/Source/Runtime/RHICore/Internal/RHICoreResourceCollection.h
@@ -54,11 +54,16 @@ namespace UE::RHICore
 		}
 	}
 
-	inline void FillResourceCollectionMemory(TRHIBufferInitializer<uint32>& Destination, TConstArrayView<FRHIResourceCollectionMember> InValues)
+	inline void FillResourceCollectionMemory(uint32* Destination, TConstArrayView<FRHIResourceCollectionMember> InValues)
 	{
 		Destination[0] = static_cast<uint32>(InValues.Num());
 
-		FillResourceCollectionUpdateMemory(Destination.GetWritableData() + 1, InValues);
+		FillResourceCollectionUpdateMemory(Destination + 1, InValues);
+	}
+
+	inline void FillResourceCollectionMemory(TRHIBufferInitializer<uint32>& Destination, TConstArrayView<FRHIResourceCollectionMember> InValues)
+	{
+		FillResourceCollectionMemory(Destination.GetWritableData(), InValues);
 	}
 
 	inline FRHIBuffer* CreateResourceCollectionBuffer(FRHICommandListBase& RHICmdList, TConstArrayView<FRHIResourceCollectionMember> InMembers)
@@ -129,20 +134,17 @@ namespace UE::RHICore
 		TConstArrayView<FRHIResourceCollectionMember> MemberUpdates = UE::RHICore::GetValidResourceCollectionUpdateList(ResourceCollection, StartIndex, InMemberUpdates);
 		if (!MemberUpdates.IsEmpty())
 		{
-			const uint32 UploadSize = MemberUpdates.Num() * sizeof(uint32);
-			const uint32 DestOffset = (StartIndex + 1) * sizeof(uint32);
+			// We can't reliably lock with offsets, so we will have to update the entire resource.
 
-			uint32* UploadData = reinterpret_cast<uint32*>(RHICmdList.LockBuffer(ResourceCollection->Buffer, DestOffset, UploadSize, RLM_WriteOnly));
+			ResourceCollection->UpdateMembers(StartIndex, MemberUpdates);
+			TConstArrayView<FRHIResourceCollectionMember> Members = ResourceCollection->GetMembers();
 
-			UE::RHICore::FillResourceCollectionUpdateMemory(UploadData, MemberUpdates);
+			const uint32 UploadSize = UE::RHICore::CalculateResourceCollectionMemorySize(Members);
+			uint32* UploadData = reinterpret_cast<uint32*>(RHICmdList.LockBuffer(ResourceCollection->Buffer, 0, UploadSize, RLM_WriteOnly));
 
-			RHICmdList.UnlockBuffer(ResourceCollection->Buffer);
+			UE::RHICore::FillResourceCollectionMemory(UploadData, Members);
 
-			TConstArrayView<FRHIResourceCollectionMember> CommandMemberUpdates = RHICmdList.IsTopOfPipe() ? RHICmdList.AllocArray(MemberUpdates) : MemberUpdates;
-			RHICmdList.EnqueueLambda([ResourceCollection, StartIndex, CommandMemberUpdates](FRHICommandListBase& RHICmdList)
-			{
-				ResourceCollection->UpdateMembers(StartIndex, CommandMemberUpdates);
-			});
+			RHICmdList.UnlockBuffer(ResourceCollection->Buffer);
 		}
 	}
 }
