Adding array versions of descriptor Free methods to allow higher level code to free multiple descriptors under a single lock.

#Synced-CL 48667092

diff --git a/Engine/Source/Runtime/RHICore/Private/RHIDescriptorAllocator.cpp b/Engine/Source/Runtime/RHICore/Private/RHIDescriptorAllocator.cpp
--- a/Engine/Source/Runtime/RHICore/Private/RHIDescriptorAllocator.cpp
+++ b/Engine/Source/Runtime/RHICore/Private/RHIDescriptorAllocator.cpp
@@ -71,6 +71,15 @@ TOptional<FRHIDescriptorAllocation> FRHIDescriptorAllocator::Allocate(uint32 Num
 	return AllocateInternal(NumDescriptors);
 }
 
+FRHIDescriptorHandle FRHIDescriptorAllocator::Allocate(ERHIDescriptorType InType)
+{
+	if (TOptional<FRHIDescriptorAllocation> Allocation = Allocate(1))
+	{
+		return FRHIDescriptorHandle(InType, Allocation->StartIndex);
+	}
+	return FRHIDescriptorHandle();
+}
+
 TOptional<FRHIDescriptorAllocation> FRHIDescriptorAllocator::AllocateInternal(uint32 NumDescriptors)
 {
 	if (const uint32 NumRanges = Ranges.Num(); NumRanges > 0)
@@ -104,7 +113,7 @@ TOptional<FRHIDescriptorAllocation> FRHIDescriptorAllocator::AllocateInternal(ui
 	return TOptional<FRHIDescriptorAllocation>();
 }
 
-void FRHIDescriptorAllocator::Free(FRHIDescriptorAllocation Allocation)
+void FRHIDescriptorAllocator::FreeInternal(FRHIDescriptorAllocation Allocation)
 {
 	const uint32 Offset = Allocation.StartIndex;
 	const uint32 NumDescriptors = Allocation.Count;
@@ -114,8 +123,6 @@ void FRHIDescriptorAllocator::Free(FRHIDescriptorAllocation Allocation)
 		return;
 	}
 
-	FScopeLock Lock(&CriticalSection);
-
 	const uint32 End = Offset + NumDescriptors;
 	// Binary search of the range list
 	uint32 Index0 = 0;
@@ -210,20 +217,35 @@ void FRHIDescriptorAllocator::Free(FRHIDescriptorAllocation Allocation)
 	}
 }
 
-FRHIDescriptorHandle FRHIDescriptorAllocator::Allocate(ERHIDescriptorType InType)
+void FRHIDescriptorAllocator::Free(FRHIDescriptorAllocation Allocation)
 {
-	if (TOptional<FRHIDescriptorAllocation> Allocation = Allocate(1))
+	FScopeLock Lock(&CriticalSection);
+	FreeInternal(Allocation);
+}
+
+void FRHIDescriptorAllocator::Free(TConstArrayView<FRHIDescriptorHandle> InHandles)
+{
+	if (InHandles.Num() == 0)
 	{
-		return FRHIDescriptorHandle(InType, Allocation->StartIndex);
+		return;
+	}
+
+	FScopeLock Lock(&CriticalSection);
+	for (FRHIDescriptorHandle Handle : InHandles)
+	{
+		if (Handle.IsValid())
+		{
+			FreeInternal(FRHIDescriptorAllocation(Handle.GetIndex(), 1));
+		}
 	}
-	return FRHIDescriptorHandle();
 }
 
 void FRHIDescriptorAllocator::Free(FRHIDescriptorHandle InHandle)
 {
 	if (InHandle.IsValid())
 	{
-		Free(FRHIDescriptorAllocation(InHandle.GetIndex(), 1));
+		FScopeLock Lock(&CriticalSection);
+		FreeInternal(FRHIDescriptorAllocation(InHandle.GetIndex(), 1));
 	}
 }
 
@@ -278,6 +300,21 @@ FRHIDescriptorHandle FRHIHeapDescriptorAllocator::Allocate(ERHIDescriptorType In
 	return FRHIDescriptorAllocator::Allocate(InType);
 }
 
+void FRHIHeapDescriptorAllocator::Free(TConstArrayView<FRHIDescriptorHandle> InHandles)
+{
+	if (InHandles.Num() > 0)
+	{
+#if DO_GUARD_SLOW
+		for (const FRHIDescriptorHandle& Handle : InHandles)
+		{
+			checkSlow(HandlesAllocation(Handle.GetType()));
+		}
+#endif
+
+		FRHIDescriptorAllocator::Free(InHandles);
+	}
+}
+
 void FRHIHeapDescriptorAllocator::Free(FRHIDescriptorHandle InHandle)
 {
 	if (InHandle.IsValid())
diff --git a/Engine/Source/Runtime/RHICore/Public/RHIDescriptorAllocator.h b/Engine/Source/Runtime/RHICore/Public/RHIDescriptorAllocator.h
--- a/Engine/Source/Runtime/RHICore/Public/RHIDescriptorAllocator.h
+++ b/Engine/Source/Runtime/RHICore/Public/RHIDescriptorAllocator.h
@@ -41,6 +41,7 @@ public:
 	RHICORE_API void Free(FRHIDescriptorAllocation Allocation);
 
 	RHICORE_API FRHIDescriptorHandle Allocate(ERHIDescriptorType InType);
+	RHICORE_API void Free(TConstArrayView<FRHIDescriptorHandle> InHandles);
 	RHICORE_API void Free(FRHIDescriptorHandle InHandle);
 
 	// Get the range of allocated descriptors. Useful for determining the smallest range to copy between heaps.
@@ -51,6 +52,7 @@ public:
 private:
 
 	RHICORE_API TOptional<FRHIDescriptorAllocation> AllocateInternal(uint32 NumDescriptors);
+	RHICORE_API void FreeInternal(FRHIDescriptorAllocation Allocation);
 
 	void RecordAlloc(uint32 Count)
 	{
@@ -90,6 +92,7 @@ public:
 	RHICORE_API FRHIHeapDescriptorAllocator(ERHIDescriptorTypeMask InTypeMask, uint32 InDescriptorCount, TConstArrayView<TStatId> InStats);
 
 	RHICORE_API FRHIDescriptorHandle Allocate(ERHIDescriptorType InType);
+	RHICORE_API void Free(TConstArrayView<FRHIDescriptorHandle> InHandles);
 	RHICORE_API void Free(FRHIDescriptorHandle InHandle);
 
 	RHICORE_API TOptional<FRHIDescriptorAllocation> Allocate(uint32 NumDescriptors);
