Disable Nvidia Aftermath Resource Tracking by default in Test and Shipping builds.
Projects will needs to explicitely enable it with r.GPUCrashDebugging.Aftermath.ResourceTracking 1

#Synced-CL 49027503

diff --git a/Engine/Source/Runtime/RHICore/Private/RHICoreNvidiaAftermath.cpp b/Engine/Source/Runtime/RHICore/Private/RHICoreNvidiaAftermath.cpp
--- a/Engine/Source/Runtime/RHICore/Private/RHICoreNvidiaAftermath.cpp
+++ b/Engine/Source/Runtime/RHICore/Private/RHICoreNvidiaAftermath.cpp
@@ -75,7 +75,7 @@ namespace UE::RHICore::Nvidia::Aftermath
 
 	static TAutoConsoleVariable<int32> CVarAftermath_Resources(
 		TEXT("r.GPUCrashDebugging.Aftermath.ResourceTracking"),
-		1,
+		0,
 		TEXT("Enable resource tracking for Aftermath dumps"),
 		ECVF_ReadOnly
 	);
@@ -169,6 +169,7 @@ namespace UE::RHICore::Nvidia::Aftermath
 
 	RHICORE_API bool IsEnabled() { return bEnabled; }
 	RHICORE_API bool IsShaderRegistrationEnabled() { return bEnabled && CVarAftermath_ShaderRegistration->GetInt(); }
+	RHICORE_API bool IsResourceTrackingEnabled() { return bEnabled && (Flags & GFSDK_Aftermath_FeatureFlags_EnableResourceTracking); }
 	RHICORE_API bool AreMarkersEnabled() { return bEnabled && (Flags & GFSDK_Aftermath_FeatureFlags_EnableMarkers); }
 
 	static void GFSDK_AFTERMATH_CALL Callback_GpuCrashDump    (const void* GPUCrashDumpData, const uint32_t GPUCrashDumpSize, void* UserData);
diff --git a/Engine/Source/Runtime/RHICore/Private/RHICoreTransientResourceAllocator.cpp b/Engine/Source/Runtime/RHICore/Private/RHICoreTransientResourceAllocator.cpp
--- a/Engine/Source/Runtime/RHICore/Private/RHICoreTransientResourceAllocator.cpp
+++ b/Engine/Source/Runtime/RHICore/Private/RHICoreTransientResourceAllocator.cpp
@@ -767,7 +767,7 @@ void FRHITransientResourceHeapAllocator::SetCreateMode(ERHITransientResourceCrea
 	const bool bSupportsParallelResourceCreation = GRHITransientAllocatorParallelResourceCreation
 #if NV_AFTERMATH
 		// Aftermath adds locks that serialize placed resource creation.
-		&& !UE::RHICore::Nvidia::Aftermath::IsEnabled()
+		&& !UE::RHICore::Nvidia::Aftermath::IsResourceTrackingEnabled()
 #endif
 		;
 
diff --git a/Engine/Source/Runtime/RHICore/Public/RHICoreNvidiaAftermath.h b/Engine/Source/Runtime/RHICore/Public/RHICoreNvidiaAftermath.h
--- a/Engine/Source/Runtime/RHICore/Public/RHICoreNvidiaAftermath.h
+++ b/Engine/Source/Runtime/RHICore/Public/RHICoreNvidiaAftermath.h
@@ -22,6 +22,7 @@
 		
 		RHICORE_API bool IsEnabled();
 		RHICORE_API bool IsShaderRegistrationEnabled();
+		RHICORE_API bool IsResourceTrackingEnabled();
 		RHICORE_API bool AreMarkersEnabled();
 
 		//
