Report resource residency status when logging Aftermath GPU crash information.

#rnx

#Synced-CL 50647393

diff --git a/Engine/Source/Runtime/RHICore/Private/RHICoreNvidiaAftermath.cpp b/Engine/Source/Runtime/RHICore/Private/RHICoreNvidiaAftermath.cpp
--- a/Engine/Source/Runtime/RHICore/Private/RHICoreNvidiaAftermath.cpp
+++ b/Engine/Source/Runtime/RHICore/Private/RHICoreNvidiaAftermath.cpp
@@ -519,6 +519,20 @@ namespace UE::RHICore::Nvidia::Aftermath
 		}
 	}
 
+	static const TCHAR* ResidencyStateToString(GFSDK_Aftermath_ResourceResidency State)
+	{
+		switch (State)
+		{
+		case GFSDK_Aftermath_ResourceResidency_FullyResident: return TEXT("Resident");
+		case GFSDK_Aftermath_ResourceResidency_Evicted: return TEXT("Evicted");
+		case GFSDK_Aftermath_ResourceResidency_MemoryFreed: return TEXT("Memory Freed");
+		case GFSDK_Aftermath_ResourceResidency_MemoryUnbound: return TEXT("Memory Unbound");
+		case GFSDK_Aftermath_ResourceResidency_Unknown:
+		default:
+			return TEXT("Unknown");
+		}
+	}
+
 	static void ProcessCrashDump(TArray<FCrashResult>& OutResults, const void* GPUCrashDumpData, const uint32_t GPUCrashDumpSize)
 	{
 		FCrashResult& CrashResult = OutResults.Emplace_GetRef();
@@ -709,6 +723,7 @@ namespace UE::RHICore::Nvidia::Aftermath
 							TEXT("\n\t\t\tIs Texture Heap     : %s")
 							TEXT("\n\t\t\tIs RTV/DSV Heap     : %s")
 							TEXT("\n\t\t\tPlaced Resource     : %s")
+							TEXT("\n\t\t\tResidency           : %s")
 							TEXT("\n\t\t\tWas Destroyed       : %s")
 							TEXT("\n\t\t\tCreate/Destroy Count: %d"),
 							ResIdx,
@@ -725,6 +740,7 @@ namespace UE::RHICore::Nvidia::Aftermath
 							FaultResInfo.bIsStaticTextureHeap ? TEXT("True") : TEXT("False"),
 							FaultResInfo.bIsRenderTargetOrDepthStencilViewHeap ? TEXT("True") : TEXT("False"),
 							FaultResInfo.bPlacedResource ? TEXT("True") : TEXT("False"),
+							ResidencyStateToString(FaultResInfo.residency),
 							FaultResInfo.bWasDestroyed ? TEXT("True") : TEXT("False"),
 							FaultResInfo.createDestroyTickCount
 						);
