Convert CommitSizeMax to an atomic to silence TSAN.

#Synced-CL 50783182

diff --git a/Engine/Source/Runtime/RHICore/Private/RHICoreTransientResourceAllocator.cpp b/Engine/Source/Runtime/RHICore/Private/RHICoreTransientResourceAllocator.cpp
--- a/Engine/Source/Runtime/RHICore/Private/RHICoreTransientResourceAllocator.cpp
+++ b/Engine/Source/Runtime/RHICore/Private/RHICoreTransientResourceAllocator.cpp
@@ -486,7 +486,9 @@ void FRHITransientHeap::Flush(uint64 AllocatorCycle, FRHITransientMemoryStats& O
 		OutAllocationStats->MemoryRanges.Add(MemoryRange);
 	}
 
-	CommitSizeMax = FMath::Max(CommitSize, CommitSizeMax);
+	// Atomically update the high watermark. CommitSizeMax is read/reset by GarbageCollect on the RHI thread.
+	uint64 CurrentMax = CommitSizeMax.load(std::memory_order_relaxed);
+	while (CommitSize > CurrentMax && !CommitSizeMax.compare_exchange_weak(CurrentMax, CommitSize, std::memory_order_relaxed)) {}
 
 	if (bHasDeallocations)
 	{
@@ -614,8 +616,7 @@ void FRHITransientHeapCache::GarbageCollect()
 	
 	for (FRHITransientHeap* Heap : LiveList)
 	{
-		TotalCommitSize += Initializer.bSupportsVirtualMapping ? Heap->CommitSizeMax : Heap->GetCapacity();
-		Heap->CommitSizeMax = 0;
+		TotalCommitSize += Initializer.bSupportsVirtualMapping ? Heap->CommitSizeMax.exchange(0, std::memory_order_relaxed) : Heap->GetCapacity();
 	}
 
 	TRACE_COUNTER_SET(TransientMemoryRangeCount, LiveList.Num());
diff --git a/Engine/Source/Runtime/RHICore/Public/RHICoreTransientResourceAllocator.h b/Engine/Source/Runtime/RHICore/Public/RHICoreTransientResourceAllocator.h
--- a/Engine/Source/Runtime/RHICore/Public/RHICoreTransientResourceAllocator.h
+++ b/Engine/Source/Runtime/RHICore/Public/RHICoreTransientResourceAllocator.h
@@ -577,7 +577,7 @@ private:
 
 	uint64 LastUsedGarbageCollectCycle = 0;
 	uint64 CommitSize = 0;
-	uint64 CommitSizeMax = 0;
+	std::atomic<uint64> CommitSizeMax{0};
 	uint32 AlignmentLog2;
 
 	TArray<FRHITransientHeapAllocator::FAliasingOverlap> AliasingOverlaps;
