Add a FRHIFrameFlipStatTrackingRunnable and move GetFrameStatistics on a separate thread that use IDXGIOutput::WaitForVBlank.
That is allow us to send to Inisight all the monitor Vblank in addition to the Vblank where the Swapchain flip the backbuffer.

Changed r.VsyncInformationInsights in order to show only the Frame Counter or the Present Counter if more information is needed.

#Synced-CL 51088875

diff --git a/Engine/Source/Runtime/D3D12RHI/Private/D3D12Viewport.h b/Engine/Source/Runtime/D3D12RHI/Private/D3D12Viewport.h
--- a/Engine/Source/Runtime/D3D12RHI/Private/D3D12Viewport.h
+++ b/Engine/Source/Runtime/D3D12RHI/Private/D3D12Viewport.h
@@ -21,6 +21,8 @@ class FD3D12UnorderedAccessView_RHI;
 class FD3D12SyncPoint;
 using FD3D12SyncPointRef = TRefCountPtr<FD3D12SyncPoint>;
 
+class FRHIFrameFlipStatTrackingRunnable;
+
 class FD3D12Viewport : public FRHIViewport, public FD3D12AdapterChild
 {
 public:
@@ -287,6 +289,8 @@ private:
 	// Can very rarely be modified on the RHI thread as well if present is skipped
 	mutable FCriticalSection ExpectedBackBufferIndexLock;
 #endif
+
+	FRHIFrameFlipStatTrackingRunnable* FlipStatThread = nullptr;
 };
 
 template<>
diff --git a/Engine/Source/Runtime/D3D12RHI/Private/Windows/WindowsD3D12Viewport.cpp b/Engine/Source/Runtime/D3D12RHI/Private/Windows/WindowsD3D12Viewport.cpp
--- a/Engine/Source/Runtime/D3D12RHI/Private/Windows/WindowsD3D12Viewport.cpp
+++ b/Engine/Source/Runtime/D3D12RHI/Private/Windows/WindowsD3D12Viewport.cpp
@@ -172,6 +172,9 @@ void FD3D12Viewport::Init()
 	// Resize to setup mGPU correctly.
 	Resize(SizeX, SizeY, bIsFullscreen, PixelFormat);
 
+#if UE_INPUT_LATENCY_TRACKING
+	FlipStatThread = new FRHIFrameFlipStatTrackingRunnable(SwapChain1);
+#endif
 	// Tell the window to redraw when they can.
 	// @todo: For Slate viewports, it doesn't make sense to post WM_PAINT messages (we swallow those.)
 	::PostMessageW(WindowHandle, WM_PAINT, 0, 0);
@@ -179,6 +182,10 @@ void FD3D12Viewport::Init()
 
 void FD3D12Viewport::FinalDestroyInternal()
 {
+	if (FlipStatThread)
+	{
+		delete FlipStatThread;
+	}
 }
 
 void FD3D12Viewport::ClearPresentQueue()
@@ -364,23 +371,6 @@ HRESULT FD3D12Viewport::PresentInternal(int32 SyncInterval)
 	{
 		// Ignore time spent waiting in Present. This function blocks based on GPU progress and space in the swap chain.
 		FRenderThreadIdleScope Scope(ERenderThreadIdleTypes::WaitingForGPUPresent);
-#if UE_INPUT_LATENCY_TRACKING
-		static FRHIFlipDetails LastFlipFrame;
-		static DXGI_FRAME_STATISTICS LastStats = { 0 };
-		HRESULT GetStatHR;
-		DXGI_FRAME_STATISTICS stats = { 0 };
-		while (SUCCEEDED(GetStatHR = SwapChain1->GetFrameStatistics(&stats)) && (stats.PresentCount > LastFlipFrame.PresentIndex))
-		{
-			FRHIFlipDetails NewFlipFrame;
-			NewFlipFrame.PresentIndex = stats.PresentCount;
-			NewFlipFrame.VBlankTimeInCycles = stats.SyncQPCTime.QuadPart;
-
-			RHIFrameInfoSetVsyncDetails(NewFlipFrame);
-
-			LastFlipFrame = NewFlipFrame;
-			LastStats = stats;
-		}
-#endif
 
 		HRESULT PresentHR;
 		{
diff --git a/Engine/Source/Runtime/RHI/Private/RHIUtilities.cpp b/Engine/Source/Runtime/RHI/Private/RHIUtilities.cpp
--- a/Engine/Source/Runtime/RHI/Private/RHIUtilities.cpp
+++ b/Engine/Source/Runtime/RHI/Private/RHIUtilities.cpp
@@ -22,11 +22,15 @@ RHIUtilities.cpp:
 #define USE_FRAME_OFFSET_THREAD 1
 
 #if UE_INPUT_LATENCY_TRACKING
-bool GVsyncInformationInsights = false;
+int32 GVsyncInformationInsights = 0;
 static FAutoConsoleVariableRef CVarVsyncInformationInsights(
 	TEXT("r.VsyncInformationInsights"),
 	GVsyncInformationInsights,
-	TEXT("Whether to show Vsync and Input events in UnrealInsights"));
+	TEXT("Whether to show Vsync and Input events in UnrealInsights:")
+	TEXT("0 - No information (default)")
+	TEXT("1 - Input / Vsync frame counter information")
+	TEXT("2 - Input / Vsync frame counter and present counter information")
+);
 
 static FRHIFrameInfo GRHIFrameInfo;
 #endif
@@ -781,6 +785,7 @@ void FRHIFrameInfo::SetVsyncDetails(const FRHIFlipDetails& NewFlipFrame)
 			{
 				//if VBlankTimeInCycles is not set we compute it from VBlankTimeInSeconds
 				uint64 VBlankTimeInCycle = (NewFlipFrame.VBlankTimeInCycles != 0) ? NewFlipFrame.VBlankTimeInCycles : ((uint64)(NewFlipFrame.VBlankTimeInSeconds / FPlatformTime::GetSecondsPerCycle64()));
+				float Hz = NewFlipFrame.Hz;
 
 				Elem.Value.VSyncCycle = VBlankTimeInCycle;
 				//LatestFrameCounter = Elem.Value.FrameIndex;
@@ -801,10 +806,31 @@ void FRHIFrameInfo::SetVsyncDetails(const FRHIFlipDetails& NewFlipFrame)
 				Elem.Value.InputLatencyMS = FPlatformTime::ToMilliseconds64(CurrentInputLatency);
 
 #if UE_INPUT_LATENCY_TRACKING
-				if (GVsyncInformationInsights)
+				if (GVsyncInformationInsights == 1)
+				{
+					TRACE_BOOKMARK_CYCLES(Elem.Value.InputCycle, TEXT("Input Frame: %" UINT64_FMT), LatestFrameCounter)
+
+					if (Hz)
+					{
+						TRACE_BOOKMARK_CYCLES(VBlankTimeInCycle, TEXT("Vsync (%.2f Hz) Frame: %" UINT64_FMT), Hz, LatestFrameCounter)
+					}
+					else
+					{
+						TRACE_BOOKMARK_CYCLES(VBlankTimeInCycle, TEXT("Vsync Frame: %" UINT64_FMT), LatestFrameCounter)
+					}
+				}
+				else if (GVsyncInformationInsights == 2)
 				{
 					TRACE_BOOKMARK_CYCLES(Elem.Value.InputCycle, TEXT("Input P:%" UINT64_FMT " F:%" UINT64_FMT), Elem.Value.PresentCounter, LatestFrameCounter)
-					TRACE_BOOKMARK_CYCLES(VBlankTimeInCycle, TEXT("Vsync P:%" UINT64_FMT " F:%" UINT64_FMT " Lat:%.2f"), Elem.Value.PresentCounter, LatestFrameCounter, Elem.Value.InputLatencyMS)
+
+					if (Hz)
+					{
+						TRACE_BOOKMARK_CYCLES(VBlankTimeInCycle, TEXT("Vsync (%.2f Hz) P:%" UINT64_FMT " F:%" UINT64_FMT " Lat:%.2f"), Hz, Elem.Value.PresentCounter, LatestFrameCounter, Elem.Value.InputLatencyMS)
+					}
+					else
+					{
+						TRACE_BOOKMARK_CYCLES(VBlankTimeInCycle, TEXT("Vsync P:%" UINT64_FMT " F:%" UINT64_FMT " Lat:%.2f"), Elem.Value.PresentCounter, LatestFrameCounter, Elem.Value.InputLatencyMS)
+					}
 				}
 #endif
 				break;
diff --git a/Engine/Source/Runtime/RHI/Public/DynamicRHI.h b/Engine/Source/Runtime/RHI/Public/DynamicRHI.h
--- a/Engine/Source/Runtime/RHI/Public/DynamicRHI.h
+++ b/Engine/Source/Runtime/RHI/Public/DynamicRHI.h
@@ -18,6 +18,7 @@ DynamicRHI.h: Dynamically bound Render Hardware Interface definitions.
 #include "Containers/ArrayView.h"
 #include "Misc/EnumClassFlags.h"
 #include "Async/TaskGraphInterfaces.h"
+#include "RHIStats.h"
 
 class FBlendStateInitializerRHI;
 class FGraphicsPipelineStateInitializer;
@@ -44,29 +45,6 @@ struct FResourceArrayUploadInterface;
 typedef TArray<FScreenResolutionRHI> FScreenResolutionArray;
 using FDisplayInformationArray = TArray<FDisplayInformation>;
 
-/** Struct to provide details of swap chain flips */
-struct FRHIFlipDetails
-{
-	uint64 PresentIndex;
-	double FlipTimeInSeconds;
-	double VBlankTimeInSeconds;
-	uint64 VBlankTimeInCycles;
-
-	FRHIFlipDetails()
-		: PresentIndex(0)
-		, FlipTimeInSeconds(0)
-		, VBlankTimeInSeconds(0)
-		, VBlankTimeInCycles(0)
-	{}
-
-	FRHIFlipDetails(uint64 InPresentIndex, double InFlipTimeInSeconds, double InVBlankTimeInSeconds, uint64 InVBlankTimeInCycles)
-		: PresentIndex(InPresentIndex)
-		, FlipTimeInSeconds(InFlipTimeInSeconds)
-		, VBlankTimeInSeconds(InVBlankTimeInSeconds)
-		, VBlankTimeInCycles(InVBlankTimeInCycles)
-	{}
-};
-
 //UE_DEPRECATED(5.3, "Use the CreateShaderResourceView function that takes an FRHIBufferSRVCreateDesc.")
 struct FShaderResourceViewInitializer : public FRHIViewDesc::FBufferSRV::FInitializer
 {
diff --git a/Engine/Source/Runtime/RHI/Public/RHIDefinitions.h b/Engine/Source/Runtime/RHI/Public/RHIDefinitions.h
--- a/Engine/Source/Runtime/RHI/Public/RHIDefinitions.h
+++ b/Engine/Source/Runtime/RHI/Public/RHIDefinitions.h
@@ -1766,4 +1766,32 @@ struct FShaderCodeValidationUBSize
 {
 	uint16 BindPoint;
 	uint32 Size;
-};
\ No newline at end of file
+};
+
+/** Struct to provide details of swap chain flips */
+struct FRHIFlipDetails
+{
+	uint64 PresentIndex;
+	double FlipTimeInSeconds;
+	double VBlankTimeInSeconds;
+	uint64 VBlankTimeInCycles;
+	float Hz;
+
+	FRHIFlipDetails()
+		: PresentIndex(0)
+		, FlipTimeInSeconds(0)
+		, VBlankTimeInSeconds(0)
+		, VBlankTimeInCycles(0)
+		, Hz(0.f)
+	{
+	}
+
+	FRHIFlipDetails(uint64 InPresentIndex, double InFlipTimeInSeconds, double InVBlankTimeInSeconds, uint64 InVBlankTimeInCycles, float InHz)
+		: PresentIndex(InPresentIndex)
+		, FlipTimeInSeconds(InFlipTimeInSeconds)
+		, VBlankTimeInSeconds(InVBlankTimeInSeconds)
+		, VBlankTimeInCycles(InVBlankTimeInCycles)
+		, Hz(InHz)
+	{
+	}
+};
diff --git a/Engine/Source/Runtime/RHICore/Internal/DXGIUtilities.h b/Engine/Source/Runtime/RHICore/Internal/DXGIUtilities.h
--- a/Engine/Source/Runtime/RHICore/Internal/DXGIUtilities.h
+++ b/Engine/Source/Runtime/RHICore/Internal/DXGIUtilities.h
@@ -5,6 +5,9 @@
 #include "HAL/Platform.h"
 
 #if PLATFORM_MICROSOFT
+#include "HAL/Runnable.h"
+#include "HAL/RunnableThread.h"
+#include "Templates/RefCounting.h"
 
 #ifndef RHICORE_PLATFORM_DXGI_H
 	#error "Platform needs to define RHICORE_PLATFORM_DXGI_H"
@@ -493,4 +496,31 @@ namespace UE::DXGIUtilities
 	}
 } // UE::DXGIUtilities
 
+#if PLATFORM_WINDOWS
+
+class FRHIFrameFlipStatTrackingRunnable : public FRunnable
+{
+	FRunnableThread* Thread = nullptr;
+	TRefCountPtr<IDXGISwapChain1> SwapChain1;
+	bool bInitialized = false;
+	bool bRun = false;
+
+	FRHIFrameFlipStatTrackingRunnable(const FRHIFrameFlipStatTrackingRunnable&) = delete;
+	FRHIFrameFlipStatTrackingRunnable& operator=(const FRHIFrameFlipStatTrackingRunnable&) = delete;
+	FRHIFrameFlipStatTrackingRunnable(FRHIFrameFlipStatTrackingRunnable&&) = delete;
+	FRHIFrameFlipStatTrackingRunnable& operator=(FRHIFrameFlipStatTrackingRunnable&&) = delete;
+
+	RHICORE_API virtual uint32 Run() override;
+	RHICORE_API virtual void Stop() override;
+	
+	RHICORE_API void StartThread();
+	RHICORE_API void StopThread();
+
+public:
+	RHICORE_API FRHIFrameFlipStatTrackingRunnable(TRefCountPtr<IDXGISwapChain1> InSwapChain1);
+	RHICORE_API ~FRHIFrameFlipStatTrackingRunnable();
+};
+
+#endif //PLATFORM_WINDOWS
+
 #endif // PLATFORM_MICROSOFT
diff --git a/Engine/Source/Runtime/RHICore/Private/DXGIUtilities.cpp b/Engine/Source/Runtime/RHICore/Private/DXGIUtilities.cpp
--- a/Engine/Source/Runtime/RHICore/Private/DXGIUtilities.cpp
+++ b/Engine/Source/Runtime/RHICore/Private/DXGIUtilities.cpp
@@ -3,6 +3,7 @@
 #include "DXGIUtilities.h"
 
 #include "RHIStats.h"
+#include "RHIUtilities.h"
 #include "CoreMinimal.h"
 
 #if PLATFORM_MICROSOFT
@@ -139,4 +140,166 @@ HRESULT UE::DXGIUtilities::GetD3DMemoryStats(IDXGIAdapter* Adapter, FRHIMemorySt
 #endif
 }
 
+#if PLATFORM_WINDOWS
+
+FRHIFrameFlipStatTrackingRunnable::FRHIFrameFlipStatTrackingRunnable(TRefCountPtr<IDXGISwapChain1> InSwapChain1)
+: SwapChain1(MoveTemp(InSwapChain1))
+{
+	StartThread();
+}
+
+FRHIFrameFlipStatTrackingRunnable::~FRHIFrameFlipStatTrackingRunnable()
+{
+	StopThread();
+}
+
+
+void FRHIFrameFlipStatTrackingRunnable::StartThread()
+{
+	check(Thread == nullptr);
+	bInitialized = true;
+	bRun = true;
+	
+	Thread = FRunnableThread::Create(this, TEXT("RHIFrameFlipStatTrackingRunnable"), 0, TPri_Normal);
+}
+
+void FRHIFrameFlipStatTrackingRunnable::StopThread()
+{
+	if (!bInitialized)
+	{
+		return;
+	}
+
+	Stop();
+
+	bInitialized = false;
+
+	if (Thread)
+	{
+		Thread->Kill(true);
+		delete Thread;
+		Thread = nullptr;
+	}
+}
+
+uint32 FRHIFrameFlipStatTrackingRunnable::Run()
+{
+	FRHIFlipDetails LastFlipFrame;
+	DXGI_FRAME_STATISTICS LastStats = { 0 };
+
+	struct FVsyncEntry
+	{
+		uint64 Cycles;
+		float Hz;
+	};
+
+	if (!SwapChain1.IsValid())
+	{
+		return 0;
+	}
+
+	TRefCountPtr<IDXGIOutput> Output;
+	SwapChain1->GetContainingOutput(Output.GetInitReference());
+
+	uint64 LastVSync = 0;
+	float Hz = 0.f;
+	const int32 VsyncDelayFrames = 12;
+	static constexpr double VsyncMatchEpsilonSeconds = 0.0003; // 0.3ms
+	const uint64 VsyncMatchEpsilon = (uint64)(VsyncMatchEpsilonSeconds / FPlatformTime::GetSecondsPerCycle64());
+
+	TArray<FVsyncEntry> VsyncQueue;
+	VsyncQueue.Reserve(VsyncDelayFrames * 2);
+	
+	IConsoleVariable* CVarVsyncInfo = IConsoleManager::Get().FindConsoleVariable(TEXT("r.VsyncInformationInsights"));
+
+	while (bRun)
+	{
+		if (!Output.IsValid())
+		{ 
+			// Something wrong in the IDXGIOutput, wait and try to reacquire the information from the SwapChain again
+			FPlatformProcess::Sleep(1.0f);
+			SwapChain1->GetContainingOutput(Output.GetInitReference());
+		}
+			
+		if (Output.IsValid())
+		{
+			// Blocking Wait for the next monitor VBlank, yielding the CPU (not busy-waiting)
+			TRACE_CPUPROFILER_EVENT_SCOPE(IDXGIOutput::WaitForVBlank_NoBusyWait);
+			HRESULT WaitForVblankHR = Output->WaitForVBlank();
+
+			// Valid IDXGIOutput and successful waiting, compute the stats
+			if (WaitForVblankHR == S_OK)
+			{
+				uint64 VsyncCycles = FPlatformTime::Cycles64();
+				uint64 delta = VsyncCycles - LastVSync;
+				LastVSync = VsyncCycles;
+				//Hz = uint32(1000.0 / FPlatformTime::ToMilliseconds64(delta));
+				Hz = (float)(1.0 / (double(delta) * FPlatformTime::GetSecondsPerCycle64()));
+
+				VsyncQueue.Add({ VsyncCycles, Hz });
+
+				// Emit bookmarks for entries that have aged past the delay window
+				// Entries that have aged past the delay window they are always at the front of the queue since entries are chronological
+				while (VsyncQueue.Num() > VsyncDelayFrames)
+				{
+#if UE_INPUT_LATENCY_TRACKING
+					if ((CVarVsyncInfo) && (CVarVsyncInfo->GetInt()))
+					{
+						TRACE_BOOKMARK_CYCLES(VsyncQueue[0].Cycles, TEXT("Vsync (%.2f Hz) Frame: Missed"), VsyncQueue[0].Hz);
+					}
+#endif
+
+					VsyncQueue.RemoveAt(0, EAllowShrinking::No);
+				}
+
+				HRESULT GetStatHR;
+				DXGI_FRAME_STATISTICS stats = { 0 };
+
+				if (SUCCEEDED(GetStatHR = SwapChain1->GetFrameStatistics(&stats)) && (stats.PresentCount > LastFlipFrame.PresentIndex))
+				{
+					FRHIFlipDetails NewFlipFrame;
+					NewFlipFrame.PresentIndex = stats.PresentCount;
+					NewFlipFrame.VBlankTimeInCycles = stats.SyncQPCTime.QuadPart;
+					NewFlipFrame.Hz = Hz;
+
+					// Remove the matching entry from the queue so it won't emit the bookmark, a more detailed bookmark is emitted in RHIFrameInfoSetVsyncDetails 
+					for (int32 i = 0; i < VsyncQueue.Num(); ++i)
+					{
+						uint64 EntryCycles = VsyncQueue[i].Cycles;
+						uint64 FlipCycles = static_cast<uint64>(NewFlipFrame.VBlankTimeInCycles);
+
+						uint64 Diff = (NewFlipFrame.VBlankTimeInCycles > VsyncQueue[i].Cycles) ? NewFlipFrame.VBlankTimeInCycles - VsyncQueue[i].Cycles : VsyncQueue[i].Cycles - NewFlipFrame.VBlankTimeInCycles;
+
+						if (Diff <= VsyncMatchEpsilon)
+						{
+							VsyncQueue.RemoveAt(i, EAllowShrinking::No);
+							break;
+						}
+					}
+
+					RHIFrameInfoSetVsyncDetails(NewFlipFrame);
+
+					LastFlipFrame = NewFlipFrame;
+					LastStats = stats;
+				}
+			}
+			else
+			{
+				// Force a reacquire of the IDXGIOutput
+				Output = nullptr;
+			}
+		}
+	}
+
+	return 0;
+}
+
+
+void FRHIFrameFlipStatTrackingRunnable::Stop()
+{
+	bRun = false;
+}
+
+#endif // PLATFORM_WINDOWS
+
 #endif // PLATFORM_MICROSOFT
diff --git a/Engine/Source/Runtime/Windows/D3D11RHI/Private/D3D11Viewport.cpp b/Engine/Source/Runtime/Windows/D3D11RHI/Private/D3D11Viewport.cpp
--- a/Engine/Source/Runtime/Windows/D3D11RHI/Private/D3D11Viewport.cpp
+++ b/Engine/Source/Runtime/Windows/D3D11RHI/Private/D3D11Viewport.cpp
@@ -210,6 +210,11 @@ FD3D11Viewport::~FD3D11Viewport()
 {
 	check(IsInRHIThread() || IsInRenderingThread());
 
+	if (FlipStatThread)
+	{
+		delete FlipStatThread;
+	}
+
 	// If the swap chain was in fullscreen mode, switch back to windowed before releasing the swap chain.
 	// DXGI throws an error otherwise.
 	if (SwapChain)
@@ -365,27 +370,6 @@ bool FD3D11Viewport::PresentChecked(IRHICommandContext& RHICmdContext, int32 Syn
 	HRESULT Result = S_OK;
 	bool bNeedNativePresent = true;
 
-#if UE_INPUT_LATENCY_TRACKING
-	if (SwapChain.IsValid())
-	{
-		static FRHIFlipDetails LastFlipFrame;
-		static DXGI_FRAME_STATISTICS LastStats = { 0 };
-		HRESULT GetStatHR;
-		DXGI_FRAME_STATISTICS stats = { 0 };
-		while (SUCCEEDED(GetStatHR = SwapChain->GetFrameStatistics(&stats)) && (stats.PresentCount > LastFlipFrame.PresentIndex))
-		{
-			FRHIFlipDetails NewFlipFrame;
-			NewFlipFrame.PresentIndex = stats.PresentCount;
-			NewFlipFrame.VBlankTimeInCycles = stats.SyncQPCTime.QuadPart;
-
-			RHIFrameInfoSetVsyncDetails(NewFlipFrame);
-
-			LastFlipFrame = NewFlipFrame;
-			LastStats = stats;
-		}
-	}
-#endif
-
 	if (IsValidRef(CustomPresent))
 	{
 		SCOPE_CYCLE_COUNTER(STAT_D3D11CustomPresentTime);
diff --git a/Engine/Source/Runtime/Windows/D3D11RHI/Private/Windows/WindowsD3D11Viewport.cpp b/Engine/Source/Runtime/Windows/D3D11RHI/Private/Windows/WindowsD3D11Viewport.cpp
--- a/Engine/Source/Runtime/Windows/D3D11RHI/Private/Windows/WindowsD3D11Viewport.cpp
+++ b/Engine/Source/Runtime/Windows/D3D11RHI/Private/Windows/WindowsD3D11Viewport.cpp
@@ -8,6 +8,7 @@
 #include "Misc/CommandLine.h"
 #include "Windows/IDXGISwapchainProvider.h"
 #include "Features/IModularFeatures.h"
+#include "RHIUtilities.h"
 
 #include "Windows/AllowWindowsPlatformTypes.h"
 #include <dwmapi.h>
@@ -146,10 +147,10 @@ FD3D11Viewport::FD3D11Viewport(FD3D11DynamicRHI* InD3DRHI, HWND InWindowHandle,
 				SwapChainDesc1.SwapEffect = DXGI_SWAP_EFFECT_FLIP_DISCARD;
 				SwapChainDesc1.Flags = GSwapChainFlags;
 
-				IDXGISwapChain1* SwapChain1 = nullptr;
+				TRefCountPtr<IDXGISwapChain1> SwapChain1;
 				HRESULT CreateSwapChainForHwndResult = DXGISwapchainProvider ?
-					DXGISwapchainProvider->CreateSwapChainForHwnd(Factory2, D3DRHI->GetDevice(), WindowHandle, &SwapChainDesc1, nullptr, nullptr, &SwapChain1) :
-					Factory2->CreateSwapChainForHwnd(D3DRHI->GetDevice(), WindowHandle, &SwapChainDesc1, nullptr, nullptr, &SwapChain1);
+					DXGISwapchainProvider->CreateSwapChainForHwnd(Factory2, D3DRHI->GetDevice(), WindowHandle, &SwapChainDesc1, nullptr, nullptr, SwapChain1.GetInitReference()) :
+					Factory2->CreateSwapChainForHwnd(D3DRHI->GetDevice(), WindowHandle, &SwapChainDesc1, nullptr, nullptr, SwapChain1.GetInitReference());
 
 				VERIFYD3D11RESULT_EX(CreateSwapChainForHwndResult, D3DRHI->GetDevice());
 				SwapChain = SwapChain1;
@@ -184,14 +185,18 @@ FD3D11Viewport::FD3D11Viewport(FD3D11DynamicRHI* InD3DRHI, HWND InWindowHandle,
 			SwapChainDesc.Flags = GSwapChainFlags;
 			SwapChainDesc.Scaling = GSwapScaling;
 
-			IDXGISwapChain1* SwapChain1 = nullptr;
+			TRefCountPtr<IDXGISwapChain1> SwapChain1;
 			HRESULT CreateSwapChainForHwndResult = DXGISwapchainProvider ?
-				DXGISwapchainProvider->CreateSwapChainForHwnd(Factory2, D3DRHI->GetDevice(), WindowHandle, &SwapChainDesc, &FSSwapChainDesc, nullptr, &SwapChain1) :
-				Factory2->CreateSwapChainForHwnd(D3DRHI->GetDevice(), WindowHandle, &SwapChainDesc, &FSSwapChainDesc, nullptr, &SwapChain1);
+				DXGISwapchainProvider->CreateSwapChainForHwnd(Factory2, D3DRHI->GetDevice(), WindowHandle, &SwapChainDesc, &FSSwapChainDesc, nullptr, SwapChain1.GetInitReference()) :
+				Factory2->CreateSwapChainForHwnd(D3DRHI->GetDevice(), WindowHandle, &SwapChainDesc, &FSSwapChainDesc, nullptr, SwapChain1.GetInitReference());
 			if(SUCCEEDED(CreateSwapChainForHwndResult))
 			{
 				SwapChain1->QueryInterface(IID_PPV_ARGS(SwapChain.GetInitReference()));
 
+#if UE_INPUT_LATENCY_TRACKING
+				FlipStatThread = new FRHIFrameFlipStatTrackingRunnable(SwapChain1);
+#endif
+
 				RECT WindowRect = {};
 				GetWindowRect(WindowHandle, &WindowRect);
 
diff --git a/Engine/Source/Runtime/Windows/D3D11RHI/Public/D3D11Viewport.h b/Engine/Source/Runtime/Windows/D3D11RHI/Public/D3D11Viewport.h
--- a/Engine/Source/Runtime/Windows/D3D11RHI/Public/D3D11Viewport.h
+++ b/Engine/Source/Runtime/Windows/D3D11RHI/Public/D3D11Viewport.h
@@ -13,6 +13,8 @@
 #include "Windows/D3D11ThirdParty.h"
 #include "DXGIUtilities.h"
 
+class FRHIFrameFlipStatTrackingRunnable;
+
 class FD3D11Viewport : public FRHIViewport
 {
 public:
@@ -137,6 +139,8 @@ protected:
 	FCustomPresentRHIRef CustomPresent;
 
 	D3D11RHI_API DXGI_MODE_DESC SetupDXGI_MODE_DESC() const;
+
+	FRHIFrameFlipStatTrackingRunnable* FlipStatThread = nullptr;
 };
 
 template<>
