Restoring RHICore's IMPLEMENT_MODULE. It was unintentionally removed with some Aftermath code.

#Synced-CL 51768028

diff --git a/Engine/Source/Runtime/RHICore/Private/RHICore.cpp b/Engine/Source/Runtime/RHICore/Private/RHICore.cpp
--- a/Engine/Source/Runtime/RHICore/Private/RHICore.cpp
+++ b/Engine/Source/Runtime/RHICore/Private/RHICore.cpp
@@ -8,6 +8,7 @@
 #include "RHI.h"
 #include "RHIShaderPlatformConfig.h"
 
+IMPLEMENT_MODULE(FDefaultModuleImpl, RHICore);
 DEFINE_LOG_CATEGORY(LogRHICore);
 
 namespace UE
