RHICore Changes
- Removing the unused RHISetShaderParametersShared and RHISetShaderUnbindsShared functions.
- GetBindlessParameterHandle should assert when asked for a descriptor range's bindless handle.

#Synced-CL 52155971

diff --git a/Engine/Source/Runtime/RHICore/Internal/RHIShaderParametersShared.h b/Engine/Source/Runtime/RHICore/Internal/RHIShaderParametersShared.h
--- a/Engine/Source/Runtime/RHICore/Internal/RHIShaderParametersShared.h
+++ b/Engine/Source/Runtime/RHICore/Internal/RHIShaderParametersShared.h
@@ -17,122 +17,15 @@ namespace UE::RHICore
 			case FRHIShaderParameterResource::EType::UnorderedAccessView: return static_cast<FRHIUnorderedAccessView*>(Resource)->GetBindlessHandle();
 			case FRHIShaderParameterResource::EType::Sampler:             return static_cast<FRHISamplerState*>(Resource)->GetBindlessHandle();
 			case FRHIShaderParameterResource::EType::ResourceCollection:  return static_cast<FRHIResourceCollection*>(Resource)->GetBindlessHandle();
-			default:
-				checkf(false, TEXT("Unhandled resource type?"));
-				break;
-			}
-		}
-
-		return FRHIDescriptorHandle();
-	}
-
-	template<typename FContextRHI>
-	inline void SetShaderUAV(FContextRHI& Context, FRHIGraphicsShader* ShaderRHI, uint16 Index, FRHIUnorderedAccessView* UAV)
-	{
-		if (ShaderRHI->GetType() == RRT_PixelShader)
-		{
-			Context.RHISetUAVParameter(static_cast<FRHIPixelShader*>(ShaderRHI), Index, UAV);
-		}
-		else
-		{
-			checkf(false, TEXT("TShaderRHI Can't have compute shader to be set. UAVs are not supported on vertex, tessellation and geometry shaders."));
-		}
-	}
-
-	template<typename FContextRHI>
-	inline void SetShaderUAV(FContextRHI& Context, FRHIComputeShader* ShaderRHI, uint16 Index, FRHIUnorderedAccessView* UAV)
-	{
-		Context.RHISetUAVParameter(ShaderRHI, Index, UAV);
-	}
-
-	template<typename FContextRHI, typename TShaderRHI>
-	inline void RHISetShaderParametersShared(
-		FContextRHI& Context
-		, TShaderRHI* ShaderRHI
-		, TConstArrayView<uint8> InParametersData
-		, TConstArrayView<FRHIShaderParameter> InParameters
-		, TConstArrayView<FRHIShaderParameterResource> InResourceParameters
-		, TConstArrayView<FRHIShaderParameterResource> InBindlessParameters
-		, bool bBindUAVsFirst = true // UE TODO: investigate where this is absolutely necessary
-	)
-	{
-		if (InParameters.Num())
-		{
-			for (const FRHIShaderParameter& Parameter : InParameters)
-			{
-				Context.RHISetShaderParameter(ShaderRHI, Parameter.BufferIndex, Parameter.BaseIndex, Parameter.ByteSize, &InParametersData[Parameter.ByteOffset]);
-			}
-		}
-
-#if PLATFORM_SUPPORTS_BINDLESS_RENDERING
-		for (const FRHIShaderParameterResource& Parameter : InBindlessParameters)
-		{
-			const FRHIDescriptorHandle Handle = GetBindlessParameterHandle(Parameter);
-			if (Handle.IsValid())
-			{
-				const uint32 BindlessIndex = Handle.GetIndex();
-				Context.RHISetShaderParameter(ShaderRHI, 0, Parameter.Index, 4, &BindlessIndex);
-			}
-		}
-#endif
-
-		if (bBindUAVsFirst)
-		{
-			for (const FRHIShaderParameterResource& Parameter : InResourceParameters)
-			{
-				if (Parameter.Type == FRHIShaderParameterResource::EType::UnorderedAccessView)
-				{
-					SetShaderUAV(Context, ShaderRHI, Parameter.Index, static_cast<FRHIUnorderedAccessView*>(Parameter.Resource));
-				}
-			}
-		}
-
-		for (const FRHIShaderParameterResource& Parameter : InResourceParameters)
-		{
-			switch (Parameter.Type)
-			{
-			case FRHIShaderParameterResource::EType::Texture:
-				Context.RHISetShaderTexture(ShaderRHI, Parameter.Index, static_cast<FRHITexture*>(Parameter.Resource));
-				break;
-			case FRHIShaderParameterResource::EType::ResourceView:
-				Context.RHISetShaderResourceViewParameter(ShaderRHI, Parameter.Index, static_cast<FRHIShaderResourceView*>(Parameter.Resource));
-				break;
-			case FRHIShaderParameterResource::EType::UnorderedAccessView:
-				if (!bBindUAVsFirst)
-				{
-					SetShaderUAV(Context, ShaderRHI, Parameter.Index, static_cast<FRHIUnorderedAccessView*>(Parameter.Resource));
-				}
-				break;
-			case FRHIShaderParameterResource::EType::Sampler:
-				Context.RHISetShaderSampler(ShaderRHI, Parameter.Index, static_cast<FRHISamplerState*>(Parameter.Resource));
-				break;
-			case FRHIShaderParameterResource::EType::UniformBuffer:
-				Context.RHISetShaderUniformBuffer(ShaderRHI, Parameter.Index, static_cast<FRHIUniformBuffer*>(Parameter.Resource));
+			case FRHIShaderParameterResource::EType::DescriptorRange:
+				checkf(false, TEXT("DescriptorRange types do not have a bindless handle, they need to be bound in a specific manner."));
 				break;
 			default:
 				checkf(false, TEXT("Unhandled resource type?"));
 				break;
 			}
 		}
-	}
 
-	template<typename FContextRHI, typename TShaderRHI>
-	inline void RHISetShaderUnbindsShared(FContextRHI& Context, TShaderRHI* ShaderRHI, TConstArrayView<FRHIShaderParameterUnbind> InUnbinds)
-	{
-		for (const FRHIShaderParameterUnbind& Unbind : InUnbinds)
-		{
-			switch (Unbind.Type)
-			{
-			case FRHIShaderParameterUnbind::EType::ResourceView:
-				Context.RHISetShaderResourceViewParameter(ShaderRHI, Unbind.Index, nullptr);
-				break;
-			case FRHIShaderParameterUnbind::EType::UnorderedAccessView:
-				SetShaderUAV(Context, ShaderRHI, Unbind.Index, nullptr);
-				break;
-			default:
-				checkf(false, TEXT("Unhandled resource type?"));
-				break;
-			}
-		}
+		return FRHIDescriptorHandle();
 	}
 }
