Add debug tracking to FRHIDescriptorAllocator to catch double-free
- rhi.TrackOutstandingDescriptorAllocations enables tracking, on by default
- rhi.TrackDescriptorCallstacks enables callstack captures when tracking descriptor allocations, off by default

#Synced-CL 52308469

diff --git a/Engine/Source/Runtime/RHICore/Private/RHIDescriptorAllocator.cpp b/Engine/Source/Runtime/RHICore/Private/RHIDescriptorAllocator.cpp
--- a/Engine/Source/Runtime/RHICore/Private/RHIDescriptorAllocator.cpp
+++ b/Engine/Source/Runtime/RHICore/Private/RHIDescriptorAllocator.cpp
@@ -1,9 +1,30 @@
 // Copyright Epic Games, Inc. All Rights Reserved.
 
 #include "RHIDescriptorAllocator.h"
+#include "HAL/IConsoleManager.h"
+#include "HAL/PlatformStackWalk.h"
 #include "Misc/ScopeLock.h"
+#include "RHICore.h"
 #include "RHIDefinitions.h"
 
+#if DO_CHECK
+static bool GRHITrackOutstandingDescriptorAllocations = true;
+static FAutoConsoleVariableRef CVarRHITrackOutstandingDescriptorAllocations(
+	TEXT("rhi.TrackOutstandingDescriptorAllocations"),
+	GRHITrackOutstandingDescriptorAllocations,
+	TEXT("When enabled, tracks outstanding descriptor allocations to detect double-free and leak bugs."),
+	ECVF_ReadOnly
+);
+
+static bool GRHITrackDescriptorCallstacks = false;
+static FAutoConsoleVariableRef CVarRHITrackDescriptorCallstacks(
+	TEXT("rhi.TrackDescriptorCallstacks"),
+	GRHITrackDescriptorCallstacks,
+	TEXT("When enabled, captures callstacks for descriptor allocations to help diagnose double-free bugs. Requires rhi.TrackOutstandingDescriptorAllocations."),
+	ECVF_ReadOnly
+);
+#endif
+
 FRHIDescriptorAllocator::FRHIDescriptorAllocator()
 {
 }
@@ -15,6 +36,29 @@ FRHIDescriptorAllocator::FRHIDescriptorAllocator(uint32 InNumDescriptors, TConst
 
 FRHIDescriptorAllocator::~FRHIDescriptorAllocator()
 {
+#if DO_CHECK
+	if (GRHITrackOutstandingDescriptorAllocations && !DebugOutstandingAllocations.IsEmpty())
+	{
+		UE_LOG(LogRHICore, Error, TEXT("FRHIDescriptorAllocator destroyed with %d outstanding allocations. Possible descriptor leak."),
+			DebugOutstandingAllocations.Num());
+
+		for (const TPair<uint32, FDebugDescriptorAllocation>& Pair : DebugOutstandingAllocations)
+		{
+			UE_LOG(LogRHICore, Error, TEXT("  Leaked allocation: Offset=%u, Count=%u"), Pair.Key, Pair.Value.Count);
+			for (uint64 Address : Pair.Value.Callstack)
+			{
+				FProgramCounterSymbolInfo SymbolInfo;
+				FPlatformStackWalk::ProgramCounterToSymbolInfo(Address, SymbolInfo);
+				ANSICHAR HumanReadableString[1024] = {};
+				FPlatformStackWalk::SymbolInfoToHumanReadableString(SymbolInfo, HumanReadableString, UE_ARRAY_COUNT(HumanReadableString));
+				UE_LOG(LogRHICore, Error, TEXT("    %hs"), HumanReadableString);
+			}
+		}
+
+		checkf(false, TEXT("FRHIDescriptorAllocator destroyed with %d outstanding allocations. See log above for details."),
+			DebugOutstandingAllocations.Num());
+	}
+#endif
 }
 
 void FRHIDescriptorAllocator::Init(uint32 InNumDescriptors, TConstArrayView<TStatId> InStats)
@@ -104,6 +148,20 @@ TOptional<FRHIDescriptorAllocation> FRHIDescriptorAllocator::AllocateInternal(ui
 
 				RecordAlloc(NumDescriptors);
 
+#if DO_CHECK
+				if (GRHITrackOutstandingDescriptorAllocations)
+				{
+					FDebugDescriptorAllocation& DebugAlloc = DebugOutstandingAllocations.Add(First);
+					DebugAlloc.Count = NumDescriptors;
+					if (GRHITrackDescriptorCallstacks)
+					{
+						DebugAlloc.Callstack.SetNumUninitialized(16);
+						uint32 Depth = FPlatformStackWalk::CaptureStackBackTrace(reinterpret_cast<uint64*>(DebugAlloc.Callstack.GetData()), 16);
+						DebugAlloc.Callstack.SetNum(Depth);
+					}
+				}
+#endif
+
 				return FRHIDescriptorAllocation(First, NumDescriptors);
 			}
 			++Index;
@@ -123,6 +181,32 @@ void FRHIDescriptorAllocator::FreeInternal(FRHIDescriptorAllocation Allocation)
 		return;
 	}
 
+#if DO_CHECK
+	if (GRHITrackOutstandingDescriptorAllocations)
+	{
+		FDebugDescriptorAllocation RemovedAlloc;
+		if (!DebugOutstandingAllocations.RemoveAndCopyValue(Offset, RemovedAlloc))
+		{
+			UE_LOG(LogRHICore, Fatal, TEXT("FRHIDescriptorAllocator::Free - Allocation at offset %u (count %u) was not found in outstanding allocations. Possible double-free."),
+				Offset, NumDescriptors);
+		}
+		else if (RemovedAlloc.Count != NumDescriptors)
+		{
+			UE_LOG(LogRHICore, Error, TEXT("FRHIDescriptorAllocator::Free - Allocation at offset %u has mismatched count: freeing %u but allocated %u. Original allocation callstack:"),
+				Offset, NumDescriptors, RemovedAlloc.Count);
+			for (uint64 Address : RemovedAlloc.Callstack)
+			{
+				FProgramCounterSymbolInfo SymbolInfo;
+				FPlatformStackWalk::ProgramCounterToSymbolInfo(Address, SymbolInfo);
+				ANSICHAR HumanReadableString[1024] = {};
+				FPlatformStackWalk::SymbolInfoToHumanReadableString(SymbolInfo, HumanReadableString, UE_ARRAY_COUNT(HumanReadableString));
+				UE_LOG(LogRHICore, Error, TEXT("    %hs"), HumanReadableString);
+			}
+			UE_LOG(LogRHICore, Fatal, TEXT("FRHIDescriptorAllocator::Free - Mismatched allocation count."));
+		}
+	}
+#endif
+
 	const uint32 End = Offset + NumDescriptors;
 	// Binary search of the range list
 	uint32 Index0 = 0;
diff --git a/Engine/Source/Runtime/RHICore/Public/RHIDescriptorAllocator.h b/Engine/Source/Runtime/RHICore/Public/RHIDescriptorAllocator.h
--- a/Engine/Source/Runtime/RHICore/Public/RHIDescriptorAllocator.h
+++ b/Engine/Source/Runtime/RHICore/Public/RHIDescriptorAllocator.h
@@ -3,6 +3,7 @@
 #pragma once
 
 #include "Containers/Array.h"
+#include "Containers/Map.h"
 #include "HAL/CriticalSection.h"
 #include "RHIDefinitions.h"
 #include "Stats/Stats.h"
@@ -80,6 +81,16 @@ private:
 
 	FCriticalSection CriticalSection;
 
+#if DO_CHECK
+	// Debug tracking of outstanding allocations to detect double-free.
+	struct FDebugDescriptorAllocation
+	{
+		uint32 Count = 0;
+		TArray<uint64, TInlineAllocator<16>> Callstack;
+	};
+	TMap<uint32, FDebugDescriptorAllocation> DebugOutstandingAllocations;
+#endif
+
 #if STATS
 	TArray<TStatId> Stats;
 #endif
